Initial Creation of Asset Placement Mode module.

#rnx

diff --git a/Engine/Source/Editor/AssetPlacementEdMode/AssetPlacementEdMode.Build.cs b/Engine/Source/Editor/AssetPlacementEdMode/AssetPlacementEdMode.Build.cs
new file mode 100644
--- /dev/null
+++ b/Engine/Source/Editor/AssetPlacementEdMode/AssetPlacementEdMode.Build.cs
@@ -0,0 +1,23 @@
+// Copyright Epic Games, Inc. All Rights Reserved.
+
+using UnrealBuildTool;
+
+public class AssetPlacementEdMode : ModuleRules
+{
+	public AssetPlacementEdMode(ReadOnlyTargetRules Target) : base(Target)
+	{
+		PrivateDependencyModuleNames.AddRange(
+			new string[] {
+				"Core",
+				"CoreUObject",
+                "InputCore",
+				"Engine",
+				"Slate",
+				"SlateCore",
+                "EditorStyle",
+				"EditorFramework",
+				"UnrealEd",
+			}
+		);
+	}
+}
diff --git a/Engine/Source/Editor/AssetPlacementEdMode/Private/AssetPlacementEdMode.cpp b/Engine/Source/Editor/AssetPlacementEdMode/Private/AssetPlacementEdMode.cpp
new file mode 100644
--- /dev/null
+++ b/Engine/Source/Editor/AssetPlacementEdMode/Private/AssetPlacementEdMode.cpp
@@ -0,0 +1,44 @@
+// Copyright Epic Games, Inc. All Rights Reserved.
+
+#include "AssetPlacementEdMode.h"
+#include "AssetPlacementEdModeToolkit.h"
+#include "InteractiveToolManager.h"
+#include "AssetPlacementEdModeStyle.h"
+#include "AssetPlacementEdModeCommands.h"
+
+#include "Settings/LevelEditorMiscSettings.h"
+
+#define LOCTEXT_NAMESPACE "AssetPlacementEdMode"
+
+constexpr const TCHAR UAssetPlacementEdMode::AssetPlacementEdModeID[];
+
+UAssetPlacementEdMode::UAssetPlacementEdMode()
+{
+	Info = FEditorModeInfo(UAssetPlacementEdMode::AssetPlacementEdModeID,
+		LOCTEXT("AssetPlacementEdModeName", "Placement"),
+		FSlateIcon(FAssetPlacementEdModeStyle::GetStyleSetName(), "PlacementBrowser.ShowAllContent"),
+		GetDefault<ULevelEditorMiscSettings>()->bEnableAssetPlacementMode);
+}
+
+UAssetPlacementEdMode::~UAssetPlacementEdMode()
+{
+}
+
+void UAssetPlacementEdMode::Enter()
+{
+	UEdMode::Enter();
+
+	const FAssetPlacementEdModeCommands& PlacementModeCommands = FAssetPlacementEdModeCommands::Get();
+}
+
+void UAssetPlacementEdMode::CreateToolkit()
+{
+	Toolkit = MakeShareable(new FAssetPlacementEdModeToolkit);
+}
+
+TMap<FName, TArray<TSharedPtr<FUICommandInfo>>> UAssetPlacementEdMode::GetModeCommands() const
+{
+	return FAssetPlacementEdModeCommands::Get().GetCommands();
+}
+
+#undef LOCTEXT_NAMESPACE
diff --git a/Engine/Source/Editor/AssetPlacementEdMode/Private/AssetPlacementEdModeCommands.cpp b/Engine/Source/Editor/AssetPlacementEdMode/Private/AssetPlacementEdModeCommands.cpp
new file mode 100644
--- /dev/null
+++ b/Engine/Source/Editor/AssetPlacementEdMode/Private/AssetPlacementEdModeCommands.cpp
@@ -0,0 +1,50 @@
+// Copyright Epic Games, Inc. All Rights Reserved.
+
+#include "AssetPlacementEdModeCommands.h"
+#include "AssetPlacementEdModeStyle.h"
+
+#define LOCTEXT_NAMESPACE "FAssetPlacementEdMode"
+
+FAssetPlacementEdModeCommands::FAssetPlacementEdModeCommands()
+	: TCommands<FAssetPlacementEdModeCommands>("AssetPlacementEdMode",
+		NSLOCTEXT("AssetPlacementEdMode", "AssetPlacementEdModeCommands", "Sample Tools Editor Mode"),
+		NAME_None,
+		FAssetPlacementEdModeStyle::GetStyleSetName())
+{
+}
+
+void FAssetPlacementEdModeCommands::RegisterCommands()
+{
+	TArray <TSharedPtr<FUICommandInfo>>& ToolCommands = Commands.FindOrAdd(NAME_Default);
+	UI_COMMAND(Select, "Select", "Select", EUserInterfaceActionType::ToggleButton, FInputChord());
+	ToolCommands.Add(Select);
+	UI_COMMAND(SelectAll, "All", "Select All Assets from the Palette", EUserInterfaceActionType::ToggleButton, FInputChord());
+	ToolCommands.Add(SelectAll);
+	UI_COMMAND(Deselect, "Deselect", "Deselect All Assets from the Palette", EUserInterfaceActionType::ToggleButton, FInputChord());
+	ToolCommands.Add(Deselect);
+	UI_COMMAND(SelectInvalid, "Invalid", "Select Invalid Assets", EUserInterfaceActionType::ToggleButton, FInputChord());
+	ToolCommands.Add(SelectInvalid);
+	UI_COMMAND(LassoSelect, "Lasso", "Select Assets with Lasso", EUserInterfaceActionType::ToggleButton, FInputChord());
+	ToolCommands.Add(LassoSelect);
+	UI_COMMAND(Paint, "Paint", "Paint Assets  from the Palette", EUserInterfaceActionType::ToggleButton, FInputChord());
+	ToolCommands.Add(Paint);
+	UI_COMMAND(Reapply, "Reapply", "Reapply Settings to Assets from the Palette", EUserInterfaceActionType::ToggleButton, FInputChord());
+	ToolCommands.Add(Reapply);
+	UI_COMMAND(PlaceSingle, "Single", "Place a Single Asset from the Palette", EUserInterfaceActionType::ToggleButton, FInputChord());
+	ToolCommands.Add(PlaceSingle);
+	UI_COMMAND(Fill, "Fill", "Fill the Selected Target with Assets from the Palette", EUserInterfaceActionType::ToggleButton, FInputChord());
+	ToolCommands.Add(Fill);
+	UI_COMMAND(Erase, "Erase", "Erases Assets Matching the Palette", EUserInterfaceActionType::ToggleButton, FInputChord());
+	ToolCommands.Add(Erase);
+	UI_COMMAND(Delete, "Delete", "Remove the Selected Assets", EUserInterfaceActionType::ToggleButton, FInputChord());
+	ToolCommands.Add(Delete);
+	UI_COMMAND(MoveToCurrentLevel, "Move", "Move Selected Assets to the Current Level", EUserInterfaceActionType::ToggleButton, FInputChord());
+	ToolCommands.Add(MoveToCurrentLevel);
+}
+
+TMap<FName, TArray<TSharedPtr<FUICommandInfo>>> FAssetPlacementEdModeCommands::GetCommands()
+{
+	return FAssetPlacementEdModeCommands::Get().Commands;
+}
+
+#undef LOCTEXT_NAMESPACE
diff --git a/Engine/Source/Editor/AssetPlacementEdMode/Private/AssetPlacementEdModeModule.cpp b/Engine/Source/Editor/AssetPlacementEdMode/Private/AssetPlacementEdModeModule.cpp
new file mode 100644
--- /dev/null
+++ b/Engine/Source/Editor/AssetPlacementEdMode/Private/AssetPlacementEdModeModule.cpp
@@ -0,0 +1,23 @@
+// Copyright Epic Games, Inc. All Rights Reserved.
+
+#include "AssetPlacementEdModeModule.h"
+#include "AssetPlacementEdModeCommands.h"
+#include "AssetPlacementEdModeStyle.h"
+
+#define LOCTEXT_NAMESPACE "AssetPlacementEdMode"
+
+IMPLEMENT_MODULE( FAssetPlacementEdMode, AssetPlacementEdMode );
+
+void FAssetPlacementEdMode::StartupModule()
+{
+	FAssetPlacementEdModeStyle::Initialize();
+	FAssetPlacementEdModeCommands::Register();
+}
+
+void FAssetPlacementEdMode::ShutdownModule()
+{
+	FAssetPlacementEdModeCommands::Unregister();
+	FAssetPlacementEdModeStyle::Shutdown();
+}
+
+#undef LOCTEXT_NAMESPACE
diff --git a/Engine/Source/Editor/AssetPlacementEdMode/Private/AssetPlacementEdModeStyle.cpp b/Engine/Source/Editor/AssetPlacementEdMode/Private/AssetPlacementEdModeStyle.cpp
new file mode 100644
--- /dev/null
+++ b/Engine/Source/Editor/AssetPlacementEdMode/Private/AssetPlacementEdModeStyle.cpp
@@ -0,0 +1,86 @@
+// Copyright Epic Games, Inc. All Rights Reserved.
+
+#include "AssetPlacementEdModeStyle.h"
+
+#include "Styling/SlateStyle.h"
+#include "Styling/SlateStyleMacros.h"
+#include "Styling/SlateStyleRegistry.h"
+
+TSharedPtr<FAssetPlacementEdModeStyle> FAssetPlacementEdModeStyle::AssetPlacementEdModeStyle;
+
+void FAssetPlacementEdModeStyle::Initialize()
+{
+	if (AssetPlacementEdModeStyle.IsValid())
+	{
+		return;
+	}
+
+	AssetPlacementEdModeStyle = MakeShareable(new FAssetPlacementEdModeStyle);
+	AssetPlacementEdModeStyle->SetupCustomStyle();
+}
+
+void FAssetPlacementEdModeStyle::Shutdown()
+{
+	AssetPlacementEdModeStyle.Reset();
+}
+
+FName FAssetPlacementEdModeStyle::GetStyleSetName()
+{
+	return "AssetPlacementEdModeStyle";
+}
+
+FAssetPlacementEdModeStyle::FAssetPlacementEdModeStyle()
+{
+	StyleSet = MakeShareable(new FSlateStyleSet(GetStyleSetName()));
+
+	StyleSet->SetParentStyleName("EditorStyle");
+	StyleSet->SetContentRoot(FPaths::EngineContentDir() / TEXT("Editor/Slate"));
+	StyleSet->SetCoreContentRoot(FPaths::EngineContentDir() / TEXT("Slate"));
+}
+
+void FAssetPlacementEdModeStyle::SetupCustomStyle()
+{
+	auto MakeImageBrushFn = [this](const FString& RelativePath)
+	{
+		static const FVector2D Icon20x20(20.0f, 20.0f);
+		constexpr const TCHAR ImageExtension[] = TEXT(".png");
+		return new FSlateImageBrush(StyleSet->RootToContentDir(RelativePath, ImageExtension), Icon20x20);
+	};
+
+	StyleSet->Set("AssetPlacementEdMode.Select", MakeImageBrushFn("Icons/GeneralTools/Select_40x"));
+	StyleSet->Set("AssetPlacementEdMode.Select.Small", MakeImageBrushFn("Icons/GeneralTools/Select_40x"));
+	StyleSet->Set("AssetPlacementEdMode.SelectAll", MakeImageBrushFn("Icons/GeneralTools/SelectAll_40x"));
+	StyleSet->Set("AssetPlacementEdMode.SelectAll.Small", MakeImageBrushFn("Icons/GeneralTools/SelectAll_40x"));
+	StyleSet->Set("AssetPlacementEdMode.Deselect", MakeImageBrushFn("Icons/GeneralTools/Deselect_40x"));
+	StyleSet->Set("AssetPlacementEdMode.Deselect.Small", MakeImageBrushFn("Icons/GeneralTools/Deselect_40x"));
+	StyleSet->Set("AssetPlacementEdMode.SelectInvalid", MakeImageBrushFn("Icons/GeneralTools/SelectInvalid_40x"));
+	StyleSet->Set("AssetPlacementEdMode.SelectInvalid.Small", MakeImageBrushFn("Icons/GeneralTools/SelectInvalid_40x"));
+	StyleSet->Set("AssetPlacementEdMode.LassoSelect", MakeImageBrushFn("Icons/GeneralTools/Lasso_40x"));
+	StyleSet->Set("AssetPlacementEdMode.LassoSelect.Small", MakeImageBrushFn("Icons/GeneralTools/Lasso_40x"));
+	StyleSet->Set("AssetPlacementEdMode.PlaceSingle", MakeImageBrushFn("Icons/GeneralTools/Foliage_40x"));
+	StyleSet->Set("AssetPlacementEdMode.PlaceSingle.Small", MakeImageBrushFn("Icons/GeneralTools/Foliage_40x"));
+	StyleSet->Set("AssetPlacementEdMode.Paint", MakeImageBrushFn("Icons/GeneralTools/Paint_40x"));
+	StyleSet->Set("AssetPlacementEdMode.Paint.Small", MakeImageBrushFn("Icons/GeneralTools/Paint_40x"));
+	StyleSet->Set("AssetPlacementEdMode.Reapply", MakeImageBrushFn("Icons/GeneralTools/Repaint_40x"));
+	StyleSet->Set("AssetPlacementEdMode.Reapply.Small", MakeImageBrushFn("Icons/GeneralTools/Repaint_40x"));
+	StyleSet->Set("AssetPlacementEdMode.Fill", MakeImageBrushFn("Icons/GeneralTools/PaintBucket_40x"));
+	StyleSet->Set("AssetPlacementEdMode.Fill.Small", MakeImageBrushFn("Icons/GeneralTools/PaintBucket_40x"));
+	StyleSet->Set("AssetPlacementEdMode.Delete", MakeImageBrushFn("Icons/GeneralTools/Delete_40x"));
+	StyleSet->Set("AssetPlacementEdMode.Delete.Small", MakeImageBrushFn("Icons/GeneralTools/Delete_40x"));
+	StyleSet->Set("AssetPlacementEdMode.Erase", MakeImageBrushFn("Icons/GeneralTools/Erase_40x"));
+	StyleSet->Set("AssetPlacementEdMode.Erase.Small", MakeImageBrushFn("Icons/GeneralTools/Erase_40x"));
+	StyleSet->Set("AssetPlacementEdMode.Filter", MakeImageBrushFn("Icons/GeneralTools/Filter_40x"));
+	StyleSet->Set("AssetPlacementEdMode.Filter.Small", MakeImageBrushFn("Icons/GeneralTools/Filter_40x"));
+	StyleSet->Set("AssetPlacementEdMode.Settings", MakeImageBrushFn("Icons/GeneralTools/Settings_40x"));
+	StyleSet->Set("AssetPlacementEdMode.Settings.Small", MakeImageBrushFn("Icons/GeneralTools/Settings_40x"));
+	StyleSet->Set("AssetPlacementEdMode.MoveToCurrentLevel", MakeImageBrushFn("Icons/GeneralTools/MoveToLevel_40x"));
+	StyleSet->Set("AssetPlacementEdMode.MoveToCurrentLevel.Small", MakeImageBrushFn("Icons/GeneralTools/MoveToLevel_40x"));
+
+	FSlateStyleRegistry::RegisterSlateStyle(*StyleSet.Get());
+}
+
+FAssetPlacementEdModeStyle::~FAssetPlacementEdModeStyle()
+{
+	FSlateStyleRegistry::UnRegisterSlateStyle(*StyleSet.Get());
+	StyleSet.Reset();
+}
diff --git a/Engine/Source/Editor/AssetPlacementEdMode/Private/AssetPlacementEdModeStyle.h b/Engine/Source/Editor/AssetPlacementEdMode/Private/AssetPlacementEdModeStyle.h
new file mode 100644
--- /dev/null
+++ b/Engine/Source/Editor/AssetPlacementEdMode/Private/AssetPlacementEdModeStyle.h
@@ -0,0 +1,27 @@
+// Copyright Epic Games, Inc. All Rights Reserved.
+
+#pragma once
+
+#include "CoreMinimal.h"
+
+class FSlateStyleSet;
+
+class FAssetPlacementEdModeStyle
+{
+public:
+	~FAssetPlacementEdModeStyle();
+
+	static void Initialize();
+
+	static void Shutdown();
+
+	static FName GetStyleSetName();
+
+private:
+	FAssetPlacementEdModeStyle();
+
+	void SetupCustomStyle();
+
+	TSharedPtr<FSlateStyleSet> StyleSet;
+	static TSharedPtr<FAssetPlacementEdModeStyle> AssetPlacementEdModeStyle;
+};
diff --git a/Engine/Source/Editor/AssetPlacementEdMode/Private/AssetPlacementEdModeToolkit.cpp b/Engine/Source/Editor/AssetPlacementEdMode/Private/AssetPlacementEdModeToolkit.cpp
new file mode 100644
--- /dev/null
+++ b/Engine/Source/Editor/AssetPlacementEdMode/Private/AssetPlacementEdModeToolkit.cpp
@@ -0,0 +1,35 @@
+// Copyright Epic Games, Inc. All Rights Reserved.
+
+#include "AssetPlacementEdModeToolkit.h"
+#include "AssetPlacementEdMode.h"
+
+#include "Modules/ModuleManager.h"
+
+#include "Widgets/Input/SButton.h"
+#include "Widgets/Text/STextBlock.h"
+
+#include "EditorModeManager.h"
+
+#define LOCTEXT_NAMESPACE "AssetPlacementEdModeToolkit"
+
+FAssetPlacementEdModeToolkit::FAssetPlacementEdModeToolkit()
+{
+}
+
+void FAssetPlacementEdModeToolkit::GetToolPaletteNames(TArray<FName>& PaletteNames) const
+{
+	PaletteNames.Add(NAME_Default);
+}
+
+
+FName FAssetPlacementEdModeToolkit::GetToolkitFName() const
+{
+	return FName("AssetPlacementEdMode");
+}
+
+FText FAssetPlacementEdModeToolkit::GetBaseToolkitName() const
+{
+	return LOCTEXT("DisplayName", "AssetPlacementEdMode Tool");
+}
+
+#undef LOCTEXT_NAMESPACE
diff --git a/Engine/Source/Editor/AssetPlacementEdMode/Public/AssetPlacementEdMode.h b/Engine/Source/Editor/AssetPlacementEdMode/Public/AssetPlacementEdMode.h
new file mode 100644
--- /dev/null
+++ b/Engine/Source/Editor/AssetPlacementEdMode/Public/AssetPlacementEdMode.h
@@ -0,0 +1,29 @@
+// Copyright Epic Games, Inc. All Rights Reserved.
+
+#pragma once
+
+#include "CoreMinimal.h"
+#include "Tools/UEdMode.h"
+
+#include "AssetPlacementEdMode.generated.h"
+
+UCLASS()
+class UAssetPlacementEdMode : public UEdMode
+{
+	GENERATED_BODY()
+
+public:
+	constexpr static const TCHAR AssetPlacementEdModeID[] = TEXT("EM_AssetPlacementEdMode");
+
+	UAssetPlacementEdMode();
+	virtual ~UAssetPlacementEdMode();
+
+	////////////////
+	// UEdMode interface
+	virtual void Enter() override;
+	virtual void CreateToolkit() override;
+	virtual TMap<FName, TArray<TSharedPtr<FUICommandInfo>>> GetModeCommands() const override;
+	//////////////////
+	// End of UEdMode interface
+	//////////////////
+};
diff --git a/Engine/Source/Editor/AssetPlacementEdMode/Public/AssetPlacementEdModeCommands.h b/Engine/Source/Editor/AssetPlacementEdMode/Public/AssetPlacementEdModeCommands.h
new file mode 100644
--- /dev/null
+++ b/Engine/Source/Editor/AssetPlacementEdMode/Public/AssetPlacementEdModeCommands.h
@@ -0,0 +1,31 @@
+// Copyright Epic Games, Inc. All Rights Reserved.
+
+#pragma once
+
+#include "CoreMinimal.h"
+#include "Framework/Commands/Commands.h"
+
+class FAssetPlacementEdModeCommands : public TCommands<FAssetPlacementEdModeCommands>
+{
+public:
+	FAssetPlacementEdModeCommands();
+
+	virtual void RegisterCommands() override;
+	static TMap<FName, TArray<TSharedPtr<FUICommandInfo>>> GetCommands();
+
+	TSharedPtr<FUICommandInfo> Select;
+	TSharedPtr<FUICommandInfo> SelectAll;
+	TSharedPtr<FUICommandInfo> Deselect;
+	TSharedPtr<FUICommandInfo> SelectInvalid;
+	TSharedPtr<FUICommandInfo> LassoSelect;
+	TSharedPtr<FUICommandInfo> Paint;
+	TSharedPtr<FUICommandInfo> Reapply;
+	TSharedPtr<FUICommandInfo> PlaceSingle;
+	TSharedPtr<FUICommandInfo> Fill;
+	TSharedPtr<FUICommandInfo> Erase;
+	TSharedPtr<FUICommandInfo> Delete;
+	TSharedPtr<FUICommandInfo> MoveToCurrentLevel;
+
+protected:
+	TMap<FName, TArray<TSharedPtr<FUICommandInfo>>> Commands;
+};
diff --git a/Engine/Source/Editor/AssetPlacementEdMode/Public/AssetPlacementEdModeModule.h b/Engine/Source/Editor/AssetPlacementEdMode/Public/AssetPlacementEdModeModule.h
new file mode 100644
--- /dev/null
+++ b/Engine/Source/Editor/AssetPlacementEdMode/Public/AssetPlacementEdModeModule.h
@@ -0,0 +1,13 @@
+// Copyright Epic Games, Inc. All Rights Reserved.
+
+#pragma once
+
+#include "CoreMinimal.h"
+#include "Modules/ModuleInterface.h"
+
+class FAssetPlacementEdMode : public IModuleInterface
+{
+public:
+	virtual void StartupModule();
+	virtual void ShutdownModule();
+};
diff --git a/Engine/Source/Editor/AssetPlacementEdMode/Public/AssetPlacementEdModeToolkit.h b/Engine/Source/Editor/AssetPlacementEdMode/Public/AssetPlacementEdModeToolkit.h
new file mode 100644
--- /dev/null
+++ b/Engine/Source/Editor/AssetPlacementEdMode/Public/AssetPlacementEdModeToolkit.h
@@ -0,0 +1,27 @@
+// Copyright Epic Games, Inc. All Rights Reserved.
+
+#pragma once
+
+#include "CoreMinimal.h"
+#include "Toolkits/BaseToolkit.h"
+#include "InteractiveTool.h"
+#include "AssetPlacementEdMode.h"
+
+class IDetailsView;
+
+/**
+ * This FModeToolkit just creates a basic UI panel that allows various InteractiveTools to
+ * be initialized, and a DetailsView used to show properties of the active Tool.
+ */
+class FAssetPlacementEdModeToolkit : public FModeToolkit
+{
+public:
+	FAssetPlacementEdModeToolkit();
+	
+	// FModeToolkit interface 
+	virtual void GetToolPaletteNames(TArray<FName>& PaletteNames) const override;
+
+	// IToolkit interface 
+	virtual FName GetToolkitFName() const override;
+	virtual FText GetBaseToolkitName() const override;
+};
diff --git a/Engine/Source/Editor/LevelAssetEditor/LevelAssetEditor.Build.cs b/Engine/Source/Editor/LevelAssetEditor/LevelAssetEditor.Build.cs
--- a/Engine/Source/Editor/LevelAssetEditor/LevelAssetEditor.Build.cs
+++ b/Engine/Source/Editor/LevelAssetEditor/LevelAssetEditor.Build.cs
@@ -54,6 +54,7 @@ public class LevelAssetEditor : ModuleRules
 		DynamicallyLoadedModuleNames.AddRange(
 			new string[]
 			{
+				"AssetPlacementEdMode",
 				// ... add any modules that your module loads dynamically here ...
 			}
 			);
diff --git a/Engine/Source/Editor/LevelAssetEditor/Private/LevelAssetEditor.cpp b/Engine/Source/Editor/LevelAssetEditor/Private/LevelAssetEditor.cpp
--- a/Engine/Source/Editor/LevelAssetEditor/Private/LevelAssetEditor.cpp
+++ b/Engine/Source/Editor/LevelAssetEditor/Private/LevelAssetEditor.cpp
@@ -8,6 +8,7 @@
 #include "Subsystems/AssetEditorSubsystem.h"
 #include "ToolMenus.h"
 #include "ULevelAssetEditor.h"
+#include "Settings/LevelEditorMiscSettings.h"
 
 static const FName LevelAssetEditorTabName("LevelAssetEditor");
 
@@ -21,9 +22,13 @@ void FLevelAssetEditorModule::StartupModule()
 	PluginCommands->MapAction(
 		FLevelAssetEditorCommands::Get().OpenPluginWindow,
 		FExecuteAction::CreateRaw(this, &FLevelAssetEditorModule::PluginButtonClicked),
-		FCanExecuteAction());
+		FCanExecuteAction(),
+		FIsActionChecked(),
+		FIsActionButtonVisible::CreateRaw(this, &FLevelAssetEditorModule::IsEnabled));
 
 	UToolMenus::RegisterStartupCallback(FSimpleMulticastDelegate::FDelegate::CreateRaw(this, &FLevelAssetEditorModule::RegisterMenus));
+
+	FModuleManager::Get().LoadModule("AssetPlacementEdMode");
 }
 
 void FLevelAssetEditorModule::ShutdownModule()
@@ -56,6 +61,11 @@ void FLevelAssetEditorModule::RegisterMenus()
 	}
 }
 
+bool FLevelAssetEditorModule::IsEnabled() const
+{
+	return GetDefault<ULevelEditorMiscSettings>()->bEnableExperimentalLevelEditor;
+}
+
 #undef LOCTEXT_NAMESPACE
 	
 IMPLEMENT_MODULE(FLevelAssetEditorModule, LevelAssetEditor)
\ No newline at end of file
diff --git a/Engine/Source/Editor/LevelAssetEditor/Public/LevelAssetEditor.h b/Engine/Source/Editor/LevelAssetEditor/Public/LevelAssetEditor.h
--- a/Engine/Source/Editor/LevelAssetEditor/Public/LevelAssetEditor.h
+++ b/Engine/Source/Editor/LevelAssetEditor/Public/LevelAssetEditor.h
@@ -22,6 +22,7 @@ public:
 private:
 
 	void RegisterMenus();
+	bool IsEnabled() const;
 
 private:
 	TSharedPtr<class FUICommandList> PluginCommands;
diff --git a/Engine/Source/Editor/UnrealEd/Classes/Settings/LevelEditorMiscSettings.h b/Engine/Source/Editor/UnrealEd/Classes/Settings/LevelEditorMiscSettings.h
--- a/Engine/Source/Editor/UnrealEd/Classes/Settings/LevelEditorMiscSettings.h
+++ b/Engine/Source/Editor/UnrealEd/Classes/Settings/LevelEditorMiscSettings.h
@@ -43,6 +43,14 @@ public:
 	UPROPERTY(EditAnywhere, config, Category=Editing, meta=( DisplayName = "Preserve Actor Scale on Replace" ))
 	uint32 bReplaceRespectsScale:1;
 
+	/** If enabled, the experimental level editor will be available */
+	UPROPERTY(EditAnywhere, config, AdvancedDisplay, Category = Editing)
+	bool bEnableExperimentalLevelEditor = false;
+
+	/** If enabled, the experimental asset placement mode will be available in the level editor */
+	UPROPERTY(EditAnywhere, config, AdvancedDisplay, Category = Editing, meta = (ConfigRestartRequired = true))
+	bool bEnableAssetPlacementMode = false;
+
 	/** If enabled, will avoid relabeling actors in UUnrealEdEngine::edactPasteSelected */
 	UPROPERTY(EditAnywhere, config, Category = Editing, meta = (DisplayName = "Avoid Actor Relabel on Paste Selected"))
 	uint32 bAvoidRelabelOnPasteSelected:1;
