Initial Creation of Asset Placement Mode UI for tools and commands.

# Added a save and load config on the settings class for UEdMode so mode wide settings would be sticky between runs.
# Removed FString refs in some functions so they could be bound using UObject delegate methods.

#rnx

diff --git a/Engine/Source/Editor/AssetPlacementEdMode/AssetPlacementEdMode.Build.cs b/Engine/Source/Editor/AssetPlacementEdMode/AssetPlacementEdMode.Build.cs
--- a/Engine/Source/Editor/AssetPlacementEdMode/AssetPlacementEdMode.Build.cs
+++ b/Engine/Source/Editor/AssetPlacementEdMode/AssetPlacementEdMode.Build.cs
@@ -17,6 +17,10 @@ public class AssetPlacementEdMode : ModuleRules
                 "EditorStyle",
 				"EditorFramework",
 				"UnrealEd",
+				"EditorInteractiveToolsFramework",
+				"InteractiveToolsFramework",
+				"Foliage",
+				"Landscape",
 			}
 		);
 	}
diff --git a/Engine/Source/Editor/AssetPlacementEdMode/Private/AssetPlacementEdMode.cpp b/Engine/Source/Editor/AssetPlacementEdMode/Private/AssetPlacementEdMode.cpp
--- a/Engine/Source/Editor/AssetPlacementEdMode/Private/AssetPlacementEdMode.cpp
+++ b/Engine/Source/Editor/AssetPlacementEdMode/Private/AssetPlacementEdMode.cpp
@@ -5,12 +5,21 @@
 #include "InteractiveToolManager.h"
 #include "AssetPlacementEdModeStyle.h"
 #include "AssetPlacementEdModeCommands.h"
+#include "EdModeInteractiveToolsContext.h"
+
+#include "Tools/PlacementSelectTool.h"
+#include "Tools/PlacementLassoSelectTool.h"
+#include "Tools/PlacementPlaceTool.h"
+#include "Tools/PlacementPlaceSingleTool.h"
+#include "Tools/PlacementFillTool.h"
+#include "Tools/PlacementEraseTool.h"
+#include "Tools/PlacementReapplySettingsTool.h"
 
 #include "Settings/LevelEditorMiscSettings.h"
 
 #define LOCTEXT_NAMESPACE "AssetPlacementEdMode"
 
-constexpr const TCHAR UAssetPlacementEdMode::AssetPlacementEdModeID[];
+constexpr TCHAR UAssetPlacementEdMode::AssetPlacementEdModeID[];
 
 UAssetPlacementEdMode::UAssetPlacementEdMode()
 {
@@ -18,6 +27,8 @@ UAssetPlacementEdMode::UAssetPlacementEdMode()
 		LOCTEXT("AssetPlacementEdModeName", "Placement"),
 		FSlateIcon(FAssetPlacementEdModeStyle::GetStyleSetName(), "PlacementBrowser.ShowAllContent"),
 		GetDefault<ULevelEditorMiscSettings>()->bEnableAssetPlacementMode);
+
+	SettingsClass = UAssetPlacementSettings::StaticClass();
 }
 
 UAssetPlacementEdMode::~UAssetPlacementEdMode()
@@ -29,6 +40,16 @@ void UAssetPlacementEdMode::Enter()
 	UEdMode::Enter();
 
 	const FAssetPlacementEdModeCommands& PlacementModeCommands = FAssetPlacementEdModeCommands::Get();
+	RegisterTool(PlacementModeCommands.Select, UPlacementModeSelectTool::ToolName, NewObject<UPlacementModeSelectToolBuilder>(this));
+	RegisterTool(PlacementModeCommands.LassoSelect, UPlacementModeLassoSelectTool::ToolName, NewObject<UPlacementModeLassoSelectToolBuilder>(this));
+	RegisterTool(PlacementModeCommands.Place, UPlacementModePlacementTool::ToolName, NewObject<UPlacementModePlacementToolBuilder>(this));
+	RegisterTool(PlacementModeCommands.PlaceSingle, UPlacementModePlaceSingleTool::ToolName, NewObject<UPlacementModeSelectAllToolBuilder>(this));
+	RegisterTool(PlacementModeCommands.Fill, UPlacementModeFillTool::ToolName, NewObject<UPlacementModeFillToolBuilder>(this));
+	RegisterTool(PlacementModeCommands.Erase, UPlacementModeEraseTool::ToolName, NewObject<UPlacementModeEraseToolBuilder>(this));
+	RegisterTool(PlacementModeCommands.ReapplySettings, UPlacementModeReapplySettingsTool::ToolName, NewObject<UPlacementModeReapplySettingsToolBuilder>(this));
+
+	ToolsContext->ToolManager->SelectActiveToolType(EToolSide::Mouse, UPlacementModeSelectTool::ToolName);
+	ToolsContext->ToolManager->ActivateTool(EToolSide::Mouse);
 }
 
 void UAssetPlacementEdMode::CreateToolkit()
@@ -41,4 +62,47 @@ TMap<FName, TArray<TSharedPtr<FUICommandInfo>>> UAssetPlacementEdMode::GetModeCo
 	return FAssetPlacementEdModeCommands::Get().GetCommands();
 }
 
+void UAssetPlacementEdMode::BindCommands()
+{
+	const TSharedRef<FUICommandList>& CommandList = Toolkit->GetToolkitCommands();
+	const FAssetPlacementEdModeCommands& PlacementModeCommands = FAssetPlacementEdModeCommands::Get();
+
+	CommandList->MapAction(PlacementModeCommands.SelectAll,
+		FExecuteAction::CreateUObject(this, &UAssetPlacementEdMode::SelectAssets, EPaletteFilter::ActivePaletteOnly, ESelectMode::Select),
+		FCanExecuteAction::CreateUObject(this, &UAssetPlacementEdMode::HasAnyAssetsInPalette, EPaletteFilter::ActivePaletteOnly));
+
+	CommandList->MapAction(PlacementModeCommands.SelectInvalid,
+		FExecuteAction::CreateUObject(this, &UAssetPlacementEdMode::SelectAssets, EPaletteFilter::InvalidInstances, ESelectMode::Select),
+		FCanExecuteAction::CreateUObject(this, &UAssetPlacementEdMode::HasAnyAssetsInPalette, EPaletteFilter::InvalidInstances));
+
+	CommandList->MapAction(PlacementModeCommands.Deselect,
+		FExecuteAction::CreateUObject(this, &UAssetPlacementEdMode::SelectAssets, EPaletteFilter::ActivePaletteOnly, ESelectMode::Deselect),
+		FCanExecuteAction::CreateUObject(this, &UAssetPlacementEdMode::HasAnyAssetsInPalette, EPaletteFilter::ActivePaletteOnly));
+
+	CommandList->MapAction(PlacementModeCommands.Delete,
+		FExecuteAction::CreateUObject(this, &UAssetPlacementEdMode::DeleteAssets),
+		FCanExecuteAction::CreateUObject(this, &UAssetPlacementEdMode::HasAnyAssetsInPalette, EPaletteFilter::ActivePaletteOnly));
+
+	CommandList->MapAction(PlacementModeCommands.MoveToActivePartition,
+		FExecuteAction::CreateUObject(this, &UAssetPlacementEdMode::MoveAssetToActivePartition),
+		FCanExecuteAction::CreateUObject(this, &UAssetPlacementEdMode::HasAnyAssetsInPalette, EPaletteFilter::ActivePaletteOnly));
+}
+
+void UAssetPlacementEdMode::SelectAssets(EPaletteFilter InSelectAllType, ESelectMode InSelectMode)
+{
+}
+
+void UAssetPlacementEdMode::DeleteAssets()
+{
+}
+
+void UAssetPlacementEdMode::MoveAssetToActivePartition()
+{
+}
+
+bool UAssetPlacementEdMode::HasAnyAssetsInPalette(EPaletteFilter InSelectAllType)
+{
+	return false;
+}
+
 #undef LOCTEXT_NAMESPACE
diff --git a/Engine/Source/Editor/AssetPlacementEdMode/Private/AssetPlacementEdModeCommands.cpp b/Engine/Source/Editor/AssetPlacementEdMode/Private/AssetPlacementEdModeCommands.cpp
--- a/Engine/Source/Editor/AssetPlacementEdMode/Private/AssetPlacementEdModeCommands.cpp
+++ b/Engine/Source/Editor/AssetPlacementEdMode/Private/AssetPlacementEdModeCommands.cpp
@@ -3,11 +3,11 @@
 #include "AssetPlacementEdModeCommands.h"
 #include "AssetPlacementEdModeStyle.h"
 
-#define LOCTEXT_NAMESPACE "FAssetPlacementEdMode"
+#define LOCTEXT_NAMESPACE "AssetPlacementEdMode"
 
 FAssetPlacementEdModeCommands::FAssetPlacementEdModeCommands()
 	: TCommands<FAssetPlacementEdModeCommands>("AssetPlacementEdMode",
-		NSLOCTEXT("AssetPlacementEdMode", "AssetPlacementEdModeCommands", "Sample Tools Editor Mode"),
+		LOCTEXT("AssetPlacementEdModeCommands", "AssetPlacement Editor Mode Commands"),
 		NAME_None,
 		FAssetPlacementEdModeStyle::GetStyleSetName())
 {
@@ -26,10 +26,10 @@ void FAssetPlacementEdModeCommands::RegisterCommands()
 	ToolCommands.Add(SelectInvalid);
 	UI_COMMAND(LassoSelect, "Lasso", "Select Assets with Lasso", EUserInterfaceActionType::ToggleButton, FInputChord());
 	ToolCommands.Add(LassoSelect);
-	UI_COMMAND(Paint, "Paint", "Paint Assets  from the Palette", EUserInterfaceActionType::ToggleButton, FInputChord());
-	ToolCommands.Add(Paint);
-	UI_COMMAND(Reapply, "Reapply", "Reapply Settings to Assets from the Palette", EUserInterfaceActionType::ToggleButton, FInputChord());
-	ToolCommands.Add(Reapply);
+	UI_COMMAND(Place, "Place", "Place Assets from the Palette", EUserInterfaceActionType::ToggleButton, FInputChord());
+	ToolCommands.Add(Place);
+	UI_COMMAND(ReapplySettings, "Reapply", "Reapply Settings to Assets from the Palette", EUserInterfaceActionType::ToggleButton, FInputChord());
+	ToolCommands.Add(ReapplySettings);
 	UI_COMMAND(PlaceSingle, "Single", "Place a Single Asset from the Palette", EUserInterfaceActionType::ToggleButton, FInputChord());
 	ToolCommands.Add(PlaceSingle);
 	UI_COMMAND(Fill, "Fill", "Fill the Selected Target with Assets from the Palette", EUserInterfaceActionType::ToggleButton, FInputChord());
@@ -38,8 +38,8 @@ void FAssetPlacementEdModeCommands::RegisterCommands()
 	ToolCommands.Add(Erase);
 	UI_COMMAND(Delete, "Delete", "Remove the Selected Assets", EUserInterfaceActionType::ToggleButton, FInputChord());
 	ToolCommands.Add(Delete);
-	UI_COMMAND(MoveToCurrentLevel, "Move", "Move Selected Assets to the Current Level", EUserInterfaceActionType::ToggleButton, FInputChord());
-	ToolCommands.Add(MoveToCurrentLevel);
+	UI_COMMAND(MoveToActivePartition, "Move", "Move Selected Assets to the Current World Partition", EUserInterfaceActionType::ToggleButton, FInputChord());
+	ToolCommands.Add(MoveToActivePartition);
 }
 
 TMap<FName, TArray<TSharedPtr<FUICommandInfo>>> FAssetPlacementEdModeCommands::GetCommands()
diff --git a/Engine/Source/Editor/AssetPlacementEdMode/Private/AssetPlacementEdModeStyle.cpp b/Engine/Source/Editor/AssetPlacementEdMode/Private/AssetPlacementEdModeStyle.cpp
--- a/Engine/Source/Editor/AssetPlacementEdMode/Private/AssetPlacementEdModeStyle.cpp
+++ b/Engine/Source/Editor/AssetPlacementEdMode/Private/AssetPlacementEdModeStyle.cpp
@@ -59,10 +59,10 @@ void FAssetPlacementEdModeStyle::SetupCustomStyle()
 	StyleSet->Set("AssetPlacementEdMode.LassoSelect.Small", MakeImageBrushFn("Icons/GeneralTools/Lasso_40x"));
 	StyleSet->Set("AssetPlacementEdMode.PlaceSingle", MakeImageBrushFn("Icons/GeneralTools/Foliage_40x"));
 	StyleSet->Set("AssetPlacementEdMode.PlaceSingle.Small", MakeImageBrushFn("Icons/GeneralTools/Foliage_40x"));
-	StyleSet->Set("AssetPlacementEdMode.Paint", MakeImageBrushFn("Icons/GeneralTools/Paint_40x"));
-	StyleSet->Set("AssetPlacementEdMode.Paint.Small", MakeImageBrushFn("Icons/GeneralTools/Paint_40x"));
-	StyleSet->Set("AssetPlacementEdMode.Reapply", MakeImageBrushFn("Icons/GeneralTools/Repaint_40x"));
-	StyleSet->Set("AssetPlacementEdMode.Reapply.Small", MakeImageBrushFn("Icons/GeneralTools/Repaint_40x"));
+	StyleSet->Set("AssetPlacementEdMode.Place", MakeImageBrushFn("Icons/GeneralTools/Paint_40x"));
+	StyleSet->Set("AssetPlacementEdMode.Place.Small", MakeImageBrushFn("Icons/GeneralTools/Paint_40x"));
+	StyleSet->Set("AssetPlacementEdMode.ReapplySettings", MakeImageBrushFn("Icons/GeneralTools/Repaint_40x"));
+	StyleSet->Set("AssetPlacementEdMode.ReapplySettings.Small", MakeImageBrushFn("Icons/GeneralTools/Repaint_40x"));
 	StyleSet->Set("AssetPlacementEdMode.Fill", MakeImageBrushFn("Icons/GeneralTools/PaintBucket_40x"));
 	StyleSet->Set("AssetPlacementEdMode.Fill.Small", MakeImageBrushFn("Icons/GeneralTools/PaintBucket_40x"));
 	StyleSet->Set("AssetPlacementEdMode.Delete", MakeImageBrushFn("Icons/GeneralTools/Delete_40x"));
@@ -73,8 +73,8 @@ void FAssetPlacementEdModeStyle::SetupCustomStyle()
 	StyleSet->Set("AssetPlacementEdMode.Filter.Small", MakeImageBrushFn("Icons/GeneralTools/Filter_40x"));
 	StyleSet->Set("AssetPlacementEdMode.Settings", MakeImageBrushFn("Icons/GeneralTools/Settings_40x"));
 	StyleSet->Set("AssetPlacementEdMode.Settings.Small", MakeImageBrushFn("Icons/GeneralTools/Settings_40x"));
-	StyleSet->Set("AssetPlacementEdMode.MoveToCurrentLevel", MakeImageBrushFn("Icons/GeneralTools/MoveToLevel_40x"));
-	StyleSet->Set("AssetPlacementEdMode.MoveToCurrentLevel.Small", MakeImageBrushFn("Icons/GeneralTools/MoveToLevel_40x"));
+	StyleSet->Set("AssetPlacementEdMode.MoveToActivePartition", MakeImageBrushFn("Icons/GeneralTools/MoveToLevel_40x"));
+	StyleSet->Set("AssetPlacementEdMode.MoveToActivePartition.Small", MakeImageBrushFn("Icons/GeneralTools/MoveToLevel_40x"));
 
 	FSlateStyleRegistry::RegisterSlateStyle(*StyleSet.Get());
 }
diff --git a/Engine/Source/Editor/AssetPlacementEdMode/Private/Tools/PlacementBrushToolBase.cpp b/Engine/Source/Editor/AssetPlacementEdMode/Private/Tools/PlacementBrushToolBase.cpp
new file mode 100644
--- /dev/null
+++ b/Engine/Source/Editor/AssetPlacementEdMode/Private/Tools/PlacementBrushToolBase.cpp
@@ -0,0 +1,50 @@
+// Copyright Epic Games, Inc. All Rights Reserved.
+
+#include "Tools/PlacementBrushToolBase.h"
+#include "UObject/Object.h"
+#include "ToolContextInterfaces.h"
+#include "InteractiveToolManager.h"
+#include "InstancedFoliageActor.h"
+#include "FoliageHelper.h"
+#include "Components/PrimitiveComponent.h"
+#include "Components/StaticMeshComponent.h"
+#include "Components/BrushComponent.h"
+#include "Components/ModelComponent.h"
+#include "LandscapeHeightfieldCollisionComponent.h"
+#include "FoliageInstancedStaticMeshComponent.h"
+
+bool UPlacementBrushToolBase::HitTest(const FRay& Ray, FHitResult& OutHit)
+{
+	UWorld* EditingWorld = GetToolManager()->GetWorld();
+	const FVector TraceStart(Ray.Origin);
+	const FVector TraceEnd(Ray.Origin + Ray.Direction * HALF_WORLD_MAX);
+	constexpr TCHAR* NAME_PlacementBrushTool = TEXT("PlacementBrushTool");
+
+	auto FilterFunc = [](const UPrimitiveComponent* InComponent) {
+		if (InComponent)
+		{
+			bool bFoliageOwned = InComponent->GetOwner() && FFoliageHelper::IsOwnedByFoliage(InComponent->GetOwner());
+			constexpr bool bAllowLandscape = true;
+			constexpr bool bAllowStaticMesh = true;
+			constexpr bool bAllowBSP = true;
+			constexpr bool bAllowFoliage = true;
+			constexpr bool bAllowTranslucent = true;
+
+			// Whitelist
+			bool bAllowed =
+				(bAllowLandscape && InComponent->IsA(ULandscapeHeightfieldCollisionComponent::StaticClass())) ||
+				(bAllowStaticMesh && InComponent->IsA(UStaticMeshComponent::StaticClass()) && !InComponent->IsA(UFoliageInstancedStaticMeshComponent::StaticClass()) && !bFoliageOwned) ||
+				(bAllowBSP && (InComponent->IsA(UBrushComponent::StaticClass()) || InComponent->IsA(UModelComponent::StaticClass()))) ||
+				(bAllowFoliage && (InComponent->IsA(UFoliageInstancedStaticMeshComponent::StaticClass()) || bFoliageOwned));
+
+			// Blacklist
+			bAllowed &=
+				(bAllowTranslucent || !(InComponent->GetMaterial(0) && IsTranslucentBlendMode(InComponent->GetMaterial(0)->GetBlendMode())));
+
+			return bAllowed;
+		}
+
+		return false; };
+
+	return AInstancedFoliageActor::FoliageTrace(EditingWorld, OutHit, FDesiredFoliageInstance(TraceStart, TraceEnd), NAME_PlacementBrushTool, false, FilterFunc);
+}
diff --git a/Engine/Source/Editor/AssetPlacementEdMode/Private/Tools/PlacementBrushToolBase.h b/Engine/Source/Editor/AssetPlacementEdMode/Private/Tools/PlacementBrushToolBase.h
new file mode 100644
--- /dev/null
+++ b/Engine/Source/Editor/AssetPlacementEdMode/Private/Tools/PlacementBrushToolBase.h
@@ -0,0 +1,19 @@
+// Copyright Epic Games, Inc. All Rights Reserved.
+
+#pragma once
+
+#include "BaseTools/BaseBrushTool.h"
+
+#include "PlacementBrushToolBase.generated.h"
+
+class FRay;
+struct FHitResult;
+
+UCLASS(Abstract, MinimalAPI)
+class UPlacementBrushToolBase : public UBaseBrushTool
+{
+	GENERATED_BODY()
+	
+public:
+	virtual bool HitTest(const FRay& Ray, FHitResult& OutHit) override;
+};
diff --git a/Engine/Source/Editor/AssetPlacementEdMode/Private/Tools/PlacementClickDragToolBase.cpp b/Engine/Source/Editor/AssetPlacementEdMode/Private/Tools/PlacementClickDragToolBase.cpp
new file mode 100644
--- /dev/null
+++ b/Engine/Source/Editor/AssetPlacementEdMode/Private/Tools/PlacementClickDragToolBase.cpp
@@ -0,0 +1,17 @@
+// Copyright Epic Games, Inc. All Rights Reserved.
+
+#include "Tools/PlacementClickDragToolBase.h"
+#include "UObject/Object.h"
+#include "ToolContextInterfaces.h"
+#include "InteractiveToolManager.h"
+#include "BaseGizmos/BrushStampIndicator.h"
+
+void UPlacementClickDragToolBase::Setup()
+{
+	UPlacementBrushToolBase::Setup();
+
+	RemoveToolPropertySource(BrushProperties);
+	BrushStampIndicator->bDrawRadiusCircle = false;
+	BrushProperties->BrushSize = .5f;
+	BrushProperties->BrushFalloffAmount = 1.f;
+}
diff --git a/Engine/Source/Editor/AssetPlacementEdMode/Private/Tools/PlacementClickDragToolBase.h b/Engine/Source/Editor/AssetPlacementEdMode/Private/Tools/PlacementClickDragToolBase.h
new file mode 100644
--- /dev/null
+++ b/Engine/Source/Editor/AssetPlacementEdMode/Private/Tools/PlacementClickDragToolBase.h
@@ -0,0 +1,15 @@
+// Copyright Epic Games, Inc. All Rights Reserved.
+
+#pragma once
+
+#include "Tools/PlacementBrushToolBase.h"
+
+#include "PlacementClickDragToolBase.generated.h"
+
+UCLASS(Abstract, MinimalAPI)
+class UPlacementClickDragToolBase : public UPlacementBrushToolBase
+{
+	GENERATED_BODY()
+public:
+	virtual void Setup() override;
+};
diff --git a/Engine/Source/Editor/AssetPlacementEdMode/Private/Tools/PlacementEraseTool.cpp b/Engine/Source/Editor/AssetPlacementEdMode/Private/Tools/PlacementEraseTool.cpp
new file mode 100644
--- /dev/null
+++ b/Engine/Source/Editor/AssetPlacementEdMode/Private/Tools/PlacementEraseTool.cpp
@@ -0,0 +1,18 @@
+// Copyright Epic Games, Inc. All Rights Reserved.
+
+#include "Tools/PlacementEraseTool.h"
+#include "UObject/Object.h"
+#include "ToolContextInterfaces.h"
+#include "InteractiveToolManager.h"
+
+constexpr TCHAR UPlacementModeEraseTool::ToolName[];
+
+bool UPlacementModeEraseToolBuilder::CanBuildTool(const FToolBuilderState& SceneState) const
+{
+	return true;
+}
+
+UInteractiveTool* UPlacementModeEraseToolBuilder::BuildTool(const FToolBuilderState& SceneState) const
+{
+	return NewObject<UPlacementModeEraseTool>(SceneState.ToolManager, UPlacementModeEraseTool::ToolName);
+}
diff --git a/Engine/Source/Editor/AssetPlacementEdMode/Private/Tools/PlacementEraseTool.h b/Engine/Source/Editor/AssetPlacementEdMode/Private/Tools/PlacementEraseTool.h
new file mode 100644
--- /dev/null
+++ b/Engine/Source/Editor/AssetPlacementEdMode/Private/Tools/PlacementEraseTool.h
@@ -0,0 +1,26 @@
+// Copyright Epic Games, Inc. All Rights Reserved.
+
+#pragma once
+
+#include "InteractiveToolBuilder.h"
+#include "Tools/PlacementBrushToolBase.h"
+
+#include "PlacementEraseTool.generated.h"
+
+UCLASS(Transient, MinimalAPI)
+class UPlacementModeEraseToolBuilder : public UInteractiveToolBuilder
+{
+	GENERATED_BODY()
+
+public:
+	virtual bool CanBuildTool(const FToolBuilderState& SceneState) const override;
+	virtual UInteractiveTool* BuildTool(const FToolBuilderState& SceneState) const override;
+};
+
+UCLASS(MinimalAPI)
+class UPlacementModeEraseTool : public UPlacementBrushToolBase
+{
+	GENERATED_BODY()
+public:
+	constexpr static TCHAR ToolName[] = TEXT("EraseTool");
+};
diff --git a/Engine/Source/Editor/AssetPlacementEdMode/Private/Tools/PlacementFillTool.cpp b/Engine/Source/Editor/AssetPlacementEdMode/Private/Tools/PlacementFillTool.cpp
new file mode 100644
--- /dev/null
+++ b/Engine/Source/Editor/AssetPlacementEdMode/Private/Tools/PlacementFillTool.cpp
@@ -0,0 +1,18 @@
+// Copyright Epic Games, Inc. All Rights Reserved.
+
+#include "Tools/PlacementFillTool.h"
+#include "UObject/Object.h"
+#include "ToolContextInterfaces.h"
+#include "InteractiveToolManager.h"
+
+constexpr TCHAR UPlacementModeFillTool::ToolName[];
+
+bool UPlacementModeFillToolBuilder::CanBuildTool(const FToolBuilderState& SceneState) const
+{
+	return true;
+}
+
+UInteractiveTool* UPlacementModeFillToolBuilder::BuildTool(const FToolBuilderState& SceneState) const
+{
+	return NewObject<UPlacementModeFillTool>(SceneState.ToolManager, UPlacementModeFillTool::ToolName);
+}
diff --git a/Engine/Source/Editor/AssetPlacementEdMode/Private/Tools/PlacementFillTool.h b/Engine/Source/Editor/AssetPlacementEdMode/Private/Tools/PlacementFillTool.h
new file mode 100644
--- /dev/null
+++ b/Engine/Source/Editor/AssetPlacementEdMode/Private/Tools/PlacementFillTool.h
@@ -0,0 +1,26 @@
+// Copyright Epic Games, Inc. All Rights Reserved.
+
+#pragma once
+
+#include "InteractiveToolBuilder.h"
+#include "Tools/PlacementClickDragToolBase.h"
+
+#include "PlacementFillTool.generated.h"
+
+UCLASS(Transient, MinimalAPI)
+class UPlacementModeFillToolBuilder : public UInteractiveToolBuilder
+{
+	GENERATED_BODY()
+
+public:
+	virtual bool CanBuildTool(const FToolBuilderState& SceneState) const override;
+	virtual UInteractiveTool* BuildTool(const FToolBuilderState& SceneState) const override;
+};
+
+UCLASS(MinimalAPI)
+class UPlacementModeFillTool : public UPlacementClickDragToolBase
+{
+	GENERATED_BODY()
+public:
+	constexpr static TCHAR ToolName[] = TEXT("FillTool");
+};
diff --git a/Engine/Source/Editor/AssetPlacementEdMode/Private/Tools/PlacementLassoSelectTool.cpp b/Engine/Source/Editor/AssetPlacementEdMode/Private/Tools/PlacementLassoSelectTool.cpp
new file mode 100644
--- /dev/null
+++ b/Engine/Source/Editor/AssetPlacementEdMode/Private/Tools/PlacementLassoSelectTool.cpp
@@ -0,0 +1,18 @@
+// Copyright Epic Games, Inc. All Rights Reserved.
+
+#include "Tools/PlacementLassoSelectTool.h"
+#include "UObject/Object.h"
+#include "ToolContextInterfaces.h"
+#include "InteractiveToolManager.h"
+
+constexpr TCHAR UPlacementModeLassoSelectTool::ToolName[];
+
+bool UPlacementModeLassoSelectToolBuilder::CanBuildTool(const FToolBuilderState& SceneState) const
+{
+	return true;
+}
+
+UInteractiveTool* UPlacementModeLassoSelectToolBuilder::BuildTool(const FToolBuilderState& SceneState) const
+{
+	return NewObject<UPlacementModeLassoSelectTool>(SceneState.ToolManager, UPlacementModeLassoSelectTool::ToolName);
+}
diff --git a/Engine/Source/Editor/AssetPlacementEdMode/Private/Tools/PlacementLassoSelectTool.h b/Engine/Source/Editor/AssetPlacementEdMode/Private/Tools/PlacementLassoSelectTool.h
new file mode 100644
--- /dev/null
+++ b/Engine/Source/Editor/AssetPlacementEdMode/Private/Tools/PlacementLassoSelectTool.h
@@ -0,0 +1,27 @@
+// Copyright Epic Games, Inc. All Rights Reserved.
+
+#pragma once
+
+#include "InteractiveToolBuilder.h"
+#include "Tools/PlacementBrushToolBase.h"
+
+#include "PlacementLassoSelectTool.generated.h"
+
+UCLASS(Transient, MinimalAPI)
+class UPlacementModeLassoSelectToolBuilder : public UInteractiveToolBuilder
+{
+	GENERATED_BODY()
+
+public:
+	virtual bool CanBuildTool(const FToolBuilderState& SceneState) const override;
+	virtual UInteractiveTool* BuildTool(const FToolBuilderState& SceneState) const override;
+};
+
+UCLASS(MinimalAPI)
+class UPlacementModeLassoSelectTool : public UPlacementBrushToolBase
+{
+	GENERATED_BODY()
+
+public:
+	constexpr static TCHAR ToolName[] = TEXT("LassoSelectTool");
+};
diff --git a/Engine/Source/Editor/AssetPlacementEdMode/Private/Tools/PlacementPlaceSingleTool.cpp b/Engine/Source/Editor/AssetPlacementEdMode/Private/Tools/PlacementPlaceSingleTool.cpp
new file mode 100644
--- /dev/null
+++ b/Engine/Source/Editor/AssetPlacementEdMode/Private/Tools/PlacementPlaceSingleTool.cpp
@@ -0,0 +1,18 @@
+// Copyright Epic Games, Inc. All Rights Reserved.
+
+#include "Tools/PlacementPlaceSingleTool.h"
+#include "UObject/Object.h"
+#include "ToolContextInterfaces.h"
+#include "InteractiveToolManager.h"
+
+constexpr TCHAR UPlacementModePlaceSingleTool::ToolName[];
+
+bool UPlacementModeSelectAllToolBuilder::CanBuildTool(const FToolBuilderState& SceneState) const
+{
+	return true;
+}
+
+UInteractiveTool* UPlacementModeSelectAllToolBuilder::BuildTool(const FToolBuilderState& SceneState) const
+{
+	return NewObject<UPlacementModePlaceSingleTool>(SceneState.ToolManager, UPlacementModePlaceSingleTool::ToolName);
+}
diff --git a/Engine/Source/Editor/AssetPlacementEdMode/Private/Tools/PlacementPlaceSingleTool.h b/Engine/Source/Editor/AssetPlacementEdMode/Private/Tools/PlacementPlaceSingleTool.h
new file mode 100644
--- /dev/null
+++ b/Engine/Source/Editor/AssetPlacementEdMode/Private/Tools/PlacementPlaceSingleTool.h
@@ -0,0 +1,27 @@
+// Copyright Epic Games, Inc. All Rights Reserved.
+
+#pragma once
+
+#include "InteractiveToolBuilder.h"
+#include "Tools/PlacementClickDragToolBase.h"
+
+#include "PlacementPlaceSingleTool.generated.h"
+
+UCLASS(Transient, MinimalAPI)
+class UPlacementModeSelectAllToolBuilder : public UInteractiveToolBuilder
+{
+	GENERATED_BODY()
+
+public:
+	virtual bool CanBuildTool(const FToolBuilderState& SceneState) const override;
+	virtual UInteractiveTool* BuildTool(const FToolBuilderState& SceneState) const override;
+};
+
+UCLASS(MinimalAPI)
+class UPlacementModePlaceSingleTool : public UPlacementClickDragToolBase
+{
+	GENERATED_BODY()
+
+public:
+	constexpr static TCHAR ToolName[] = TEXT("PlaceSingleTool");
+};
diff --git a/Engine/Source/Editor/AssetPlacementEdMode/Private/Tools/PlacementPlaceTool.cpp b/Engine/Source/Editor/AssetPlacementEdMode/Private/Tools/PlacementPlaceTool.cpp
new file mode 100644
--- /dev/null
+++ b/Engine/Source/Editor/AssetPlacementEdMode/Private/Tools/PlacementPlaceTool.cpp
@@ -0,0 +1,18 @@
+// Copyright Epic Games, Inc. All Rights Reserved.
+
+#include "Tools/PlacementPlaceTool.h"
+#include "UObject/Object.h"
+#include "ToolContextInterfaces.h"
+#include "InteractiveToolManager.h"
+
+constexpr TCHAR UPlacementModePlacementTool::ToolName[];
+
+bool UPlacementModePlacementToolBuilder::CanBuildTool(const FToolBuilderState& SceneState) const
+{
+	return true;
+}
+
+UInteractiveTool* UPlacementModePlacementToolBuilder::BuildTool(const FToolBuilderState& SceneState) const
+{
+	return NewObject<UPlacementModePlacementTool>(SceneState.ToolManager, UPlacementModePlacementTool::ToolName);
+}
diff --git a/Engine/Source/Editor/AssetPlacementEdMode/Private/Tools/PlacementPlaceTool.h b/Engine/Source/Editor/AssetPlacementEdMode/Private/Tools/PlacementPlaceTool.h
new file mode 100644
--- /dev/null
+++ b/Engine/Source/Editor/AssetPlacementEdMode/Private/Tools/PlacementPlaceTool.h
@@ -0,0 +1,26 @@
+// Copyright Epic Games, Inc. All Rights Reserved.
+
+#pragma once
+
+#include "InteractiveToolBuilder.h"
+#include "Tools/PlacementBrushToolBase.h"
+
+#include "PlacementPlaceTool.generated.h"
+
+UCLASS(Transient, MinimalAPI)
+class UPlacementModePlacementToolBuilder : public UInteractiveToolBuilder
+{
+	GENERATED_BODY()
+
+public:
+	virtual bool CanBuildTool(const FToolBuilderState& SceneState) const override;
+	virtual UInteractiveTool* BuildTool(const FToolBuilderState& SceneState) const override;
+};
+
+UCLASS(MinimalAPI)
+class UPlacementModePlacementTool : public UPlacementBrushToolBase
+{
+	GENERATED_BODY()
+public:
+	constexpr static TCHAR ToolName[] = TEXT("PlaceTool");
+};
diff --git a/Engine/Source/Editor/AssetPlacementEdMode/Private/Tools/PlacementReapplySettingsTool.cpp b/Engine/Source/Editor/AssetPlacementEdMode/Private/Tools/PlacementReapplySettingsTool.cpp
new file mode 100644
--- /dev/null
+++ b/Engine/Source/Editor/AssetPlacementEdMode/Private/Tools/PlacementReapplySettingsTool.cpp
@@ -0,0 +1,18 @@
+// Copyright Epic Games, Inc. All Rights Reserved.
+
+#include "Tools/PlacementReapplySettingsTool.h"
+#include "UObject/Object.h"
+#include "ToolContextInterfaces.h"
+#include "InteractiveToolManager.h"
+
+constexpr TCHAR UPlacementModeReapplySettingsTool::ToolName[];
+
+bool UPlacementModeReapplySettingsToolBuilder::CanBuildTool(const FToolBuilderState& SceneState) const
+{
+	return true;
+}
+
+UInteractiveTool* UPlacementModeReapplySettingsToolBuilder::BuildTool(const FToolBuilderState& SceneState) const
+{
+	return NewObject<UPlacementModeReapplySettingsTool>(SceneState.ToolManager, UPlacementModeReapplySettingsTool::ToolName);
+}
diff --git a/Engine/Source/Editor/AssetPlacementEdMode/Private/Tools/PlacementReapplySettingsTool.h b/Engine/Source/Editor/AssetPlacementEdMode/Private/Tools/PlacementReapplySettingsTool.h
new file mode 100644
--- /dev/null
+++ b/Engine/Source/Editor/AssetPlacementEdMode/Private/Tools/PlacementReapplySettingsTool.h
@@ -0,0 +1,26 @@
+// Copyright Epic Games, Inc. All Rights Reserved.
+
+#pragma once
+
+#include "InteractiveToolBuilder.h"
+#include "Tools/PlacementBrushToolBase.h"
+
+#include "PlacementReapplySettingsTool.generated.h"
+
+UCLASS(Transient, MinimalAPI)
+class UPlacementModeReapplySettingsToolBuilder : public UInteractiveToolBuilder
+{
+	GENERATED_BODY()
+
+public:
+	virtual bool CanBuildTool(const FToolBuilderState& SceneState) const override;
+	virtual UInteractiveTool* BuildTool(const FToolBuilderState& SceneState) const override;
+};
+
+UCLASS(MinimalAPI)
+class UPlacementModeReapplySettingsTool : public UPlacementBrushToolBase
+{
+	GENERATED_BODY()
+public:
+	constexpr static TCHAR ToolName[] = TEXT("ReapplySettings");
+};
diff --git a/Engine/Source/Editor/AssetPlacementEdMode/Private/Tools/PlacementSelectTool.cpp b/Engine/Source/Editor/AssetPlacementEdMode/Private/Tools/PlacementSelectTool.cpp
new file mode 100644
--- /dev/null
+++ b/Engine/Source/Editor/AssetPlacementEdMode/Private/Tools/PlacementSelectTool.cpp
@@ -0,0 +1,18 @@
+// Copyright Epic Games, Inc. All Rights Reserved.
+
+#include "Tools/PlacementSelectTool.h"
+#include "UObject/Object.h"
+#include "ToolContextInterfaces.h"
+#include "InteractiveToolManager.h"
+
+constexpr TCHAR UPlacementModeSelectTool::ToolName[];
+
+bool UPlacementModeSelectToolBuilder::CanBuildTool(const FToolBuilderState& SceneState) const
+{
+	return true;
+}
+
+UInteractiveTool* UPlacementModeSelectToolBuilder::BuildTool(const FToolBuilderState& SceneState) const
+{
+	return NewObject<UPlacementModeSelectTool>(SceneState.ToolManager, UPlacementModeSelectTool::ToolName);
+}
diff --git a/Engine/Source/Editor/AssetPlacementEdMode/Private/Tools/PlacementSelectTool.h b/Engine/Source/Editor/AssetPlacementEdMode/Private/Tools/PlacementSelectTool.h
new file mode 100644
--- /dev/null
+++ b/Engine/Source/Editor/AssetPlacementEdMode/Private/Tools/PlacementSelectTool.h
@@ -0,0 +1,27 @@
+// Copyright Epic Games, Inc. All Rights Reserved.
+
+#pragma once
+
+#include "InteractiveToolBuilder.h"
+#include "Tools/PlacementClickDragToolBase.h"
+
+#include "PlacementSelectTool.generated.h"
+
+UCLASS(Transient, MinimalAPI)
+class UPlacementModeSelectToolBuilder : public UInteractiveToolBuilder
+{
+	GENERATED_BODY()
+
+public:
+	virtual bool CanBuildTool(const FToolBuilderState& SceneState) const override;
+	virtual UInteractiveTool* BuildTool(const FToolBuilderState& SceneState) const override;
+};
+
+UCLASS(MinimalAPI)
+class UPlacementModeSelectTool : public UPlacementClickDragToolBase
+{
+	GENERATED_BODY()
+
+public:
+	constexpr static TCHAR ToolName[] = TEXT("SelectTool");
+};
diff --git a/Engine/Source/Editor/AssetPlacementEdMode/Public/AssetPlacementEdMode.h b/Engine/Source/Editor/AssetPlacementEdMode/Public/AssetPlacementEdMode.h
--- a/Engine/Source/Editor/AssetPlacementEdMode/Public/AssetPlacementEdMode.h
+++ b/Engine/Source/Editor/AssetPlacementEdMode/Public/AssetPlacementEdMode.h
@@ -4,9 +4,46 @@
 
 #include "CoreMinimal.h"
 #include "Tools/UEdMode.h"
+#include "TypedElementHandle.h"
 
 #include "AssetPlacementEdMode.generated.h"
 
+USTRUCT()
+struct FPaletteItem
+{
+	GENERATED_BODY()
+	FTypedHandleTypeId ElementId;
+	bool bIsEnabled;
+};
+
+UCLASS(config = EditorPerProjectUserSettings)
+class UAssetPlacementSettings : public UObject
+{
+	GENERATED_BODY()
+
+public:
+	UPROPERTY(config, EditAnywhere, Category = "Level Partition Settings")
+	bool bPlaceInCurrentLevelPartition = true;
+
+	UPROPERTY(config, EditAnywhere, Category = "Filters")
+	bool bLandscape = true;
+
+	UPROPERTY(config, EditAnywhere, Category = "Filters")
+	bool bStaticMeshes = true;
+	
+	UPROPERTY(config, EditAnywhere, Category = "Filters")
+	bool bBSP = true;
+	
+	UPROPERTY(config, EditAnywhere, Category = "Filters")
+	bool bFoliage = false;
+	
+	UPROPERTY(config, EditAnywhere, Category = "Filters")
+	bool bTranslucent = false;
+
+	UPROPERTY(config, EditAnywhere, Category = "Asset Palette")
+	TArray<FPaletteItem> PaletteItems;
+};
+
 UCLASS()
 class UAssetPlacementEdMode : public UEdMode
 {
@@ -23,7 +60,27 @@ public:
 	virtual void Enter() override;
 	virtual void CreateToolkit() override;
 	virtual TMap<FName, TArray<TSharedPtr<FUICommandInfo>>> GetModeCommands() const override;
+	virtual void BindCommands() override;
 	//////////////////
 	// End of UEdMode interface
 	//////////////////
+
+protected:
+	enum class EPaletteFilter
+	{
+		EntirePalette,
+		ActivePaletteOnly,
+		InvalidInstances,
+	};
+
+	enum class ESelectMode
+	{
+		Select,
+		Deselect,
+	};
+
+	void SelectAssets(EPaletteFilter InSelectAllType, ESelectMode InSelectMode);
+	void DeleteAssets();
+	void MoveAssetToActivePartition();
+	bool HasAnyAssetsInPalette(EPaletteFilter InSelectAllType);
 };
diff --git a/Engine/Source/Editor/AssetPlacementEdMode/Public/AssetPlacementEdModeCommands.h b/Engine/Source/Editor/AssetPlacementEdMode/Public/AssetPlacementEdModeCommands.h
--- a/Engine/Source/Editor/AssetPlacementEdMode/Public/AssetPlacementEdModeCommands.h
+++ b/Engine/Source/Editor/AssetPlacementEdMode/Public/AssetPlacementEdModeCommands.h
@@ -18,13 +18,13 @@ public:
 	TSharedPtr<FUICommandInfo> Deselect;
 	TSharedPtr<FUICommandInfo> SelectInvalid;
 	TSharedPtr<FUICommandInfo> LassoSelect;
-	TSharedPtr<FUICommandInfo> Paint;
-	TSharedPtr<FUICommandInfo> Reapply;
+	TSharedPtr<FUICommandInfo> Place;
+	TSharedPtr<FUICommandInfo> ReapplySettings;
 	TSharedPtr<FUICommandInfo> PlaceSingle;
 	TSharedPtr<FUICommandInfo> Fill;
 	TSharedPtr<FUICommandInfo> Erase;
 	TSharedPtr<FUICommandInfo> Delete;
-	TSharedPtr<FUICommandInfo> MoveToCurrentLevel;
+	TSharedPtr<FUICommandInfo> MoveToActivePartition;
 
 protected:
 	TMap<FName, TArray<TSharedPtr<FUICommandInfo>>> Commands;
diff --git a/Engine/Source/Editor/Experimental/EditorInteractiveToolsFramework/Private/EdModeInteractiveToolsContext.cpp b/Engine/Source/Editor/Experimental/EditorInteractiveToolsFramework/Private/EdModeInteractiveToolsContext.cpp
--- a/Engine/Source/Editor/Experimental/EditorInteractiveToolsFramework/Private/EdModeInteractiveToolsContext.cpp
+++ b/Engine/Source/Editor/Experimental/EditorInteractiveToolsFramework/Private/EdModeInteractiveToolsContext.cpp
@@ -582,6 +582,16 @@ void UEdModeInteractiveToolsContext::PostInvalidation()
 	InvalidationTimestamp++;
 }
 
+UWorld* UEdModeInteractiveToolsContext::GetWorld() const
+{
+	if (EditorModeManager)
+	{
+		return EditorModeManager->GetWorld();
+	}
+
+	return nullptr;
+}
+
 void UEdModeInteractiveToolsContext::Tick(FEditorViewportClient* ViewportClient, float DeltaTime)
 {
 	// invalidate this viewport if it's timestamp is not current
@@ -1130,7 +1140,7 @@ FRay UEdModeInteractiveToolsContext::GetLastWorldRay() const
 	return CurrentMouseState.Mouse.WorldRay;
 }
 
-bool UEdModeInteractiveToolsContext::CanStartTool(const FString& ToolTypeIdentifier) const
+bool UEdModeInteractiveToolsContext::CanStartTool(const FString ToolTypeIdentifier) const
 {
 	return UInteractiveToolsContext::CanStartTool(EToolSide::Mouse, ToolTypeIdentifier);
 }
@@ -1155,7 +1165,7 @@ bool UEdModeInteractiveToolsContext::CanCompleteActiveTool() const
 	return UInteractiveToolsContext::CanCompleteActiveTool(EToolSide::Mouse);
 }
 
-void UEdModeInteractiveToolsContext::StartTool(const FString& ToolTypeIdentifier)
+void UEdModeInteractiveToolsContext::StartTool(const FString ToolTypeIdentifier)
 {
 	FString LocalIdentifier(ToolTypeIdentifier);
 	PendingToolToStart = LocalIdentifier;
diff --git a/Engine/Source/Editor/Experimental/EditorInteractiveToolsFramework/Public/EdModeInteractiveToolsContext.h b/Engine/Source/Editor/Experimental/EditorInteractiveToolsFramework/Public/EdModeInteractiveToolsContext.h
--- a/Engine/Source/Editor/Experimental/EditorInteractiveToolsFramework/Public/EdModeInteractiveToolsContext.h
+++ b/Engine/Source/Editor/Experimental/EditorInteractiveToolsFramework/Public/EdModeInteractiveToolsContext.h
@@ -55,8 +55,10 @@ public:
 
 	virtual void PostInvalidation();
 
+	// UObject Interface
+	virtual UWorld* GetWorld() const override;
 
-	// call these from your FEdMode functions of the same name
+	// call these from your UEdMode functions of the same name
 
 	virtual void Tick(FEditorViewportClient* ViewportClient, float DeltaTime);
 	virtual void Render(const FSceneView* View, FViewport* Viewport, FPrimitiveDrawInterface* PDI);
@@ -79,12 +81,12 @@ public:
 	// Utility functions useful for hooking up to UICommand/etc
 	//
 
-	virtual bool CanStartTool(const FString& ToolTypeIdentifier) const;
+	virtual bool CanStartTool(const FString ToolTypeIdentifier) const;
 	virtual bool ActiveToolHasAccept() const;
 	virtual bool CanAcceptActiveTool() const;
 	virtual bool CanCancelActiveTool() const;
 	virtual bool CanCompleteActiveTool() const;
-	virtual void StartTool(const FString& ToolTypeIdentifier);
+	virtual void StartTool(const FString ToolTypeIdentifier);
 	virtual void EndTool(EToolShutdownType ShutdownType);
 
 	virtual bool ShouldIgnoreHotkeys() const { return bInFlyMode; }
diff --git a/Engine/Source/Editor/UnrealEd/Private/Tools/UEdMode.cpp b/Engine/Source/Editor/UnrealEd/Private/Tools/UEdMode.cpp
--- a/Engine/Source/Editor/UnrealEd/Private/Tools/UEdMode.cpp
+++ b/Engine/Source/Editor/UnrealEd/Private/Tools/UEdMode.cpp
@@ -237,6 +237,7 @@ void UEdMode::Enter()
 
 	if (SettingsObject)
 	{
+		SettingsObject->LoadConfig();
 		Toolkit->SetModeSettingsObject(SettingsObject);
 	}
 
@@ -249,9 +250,10 @@ void UEdMode::RegisterTool(TSharedPtr<FUICommandInfo> UICommand, FString ToolIde
 	const TSharedRef<FUICommandList>& CommandList = Toolkit->GetToolkitCommands();
 	ToolsContext->ToolManager->RegisterToolType(ToolIdentifier, Builder);
 	CommandList->MapAction(UICommand,
-		FExecuteAction::CreateLambda([this, ToolIdentifier]() { this->ToolsContext->StartTool(ToolIdentifier); }),
-		FCanExecuteAction::CreateLambda([this, ToolIdentifier]() { return this->ToolsContext->ToolManager->CanActivateTool(EToolSide::Mouse, ToolIdentifier); }),
-		FIsActionChecked::CreateLambda([this, Builder]() { return this->ToolsContext->IsToolBuilderActive(EToolSide::Mouse, Builder); }));
+		FExecuteAction::CreateUObject(ToolsContext.Get(), &UEdModeInteractiveToolsContext::StartTool, ToolIdentifier),
+		FCanExecuteAction::CreateUObject(ToolsContext->ToolManager, &UInteractiveToolManager::CanActivateTool, EToolSide::Mouse, ToolIdentifier),
+		FIsActionChecked::CreateUObject(ToolsContext.Get(), &UEdModeInteractiveToolsContext::IsToolBuilderActive, EToolSide::Mouse, Builder),
+		EUIActionRepeatMode::RepeatDisabled);
 
 	RegisteredTools.Emplace(UICommand, ToolIdentifier);
 }
@@ -261,6 +263,11 @@ void UEdMode::Exit()
 	GetToolManager()->OnToolStarted.RemoveAll(this);
 	GetToolManager()->OnToolEnded.RemoveAll(this);
 
+	if (SettingsObject)
+	{
+		SettingsObject->SaveConfig();
+	}
+
 	const TSharedRef<FUICommandList>& CommandList = Toolkit->GetToolkitCommands();
 	for (auto& RegisteredTool : RegisteredTools)
 	{
diff --git a/Engine/Source/Runtime/Experimental/InteractiveToolsFramework/Private/InteractiveToolManager.cpp b/Engine/Source/Runtime/Experimental/InteractiveToolsFramework/Private/InteractiveToolManager.cpp
--- a/Engine/Source/Runtime/Experimental/InteractiveToolsFramework/Private/InteractiveToolManager.cpp
+++ b/Engine/Source/Runtime/Experimental/InteractiveToolsFramework/Private/InteractiveToolManager.cpp
@@ -97,7 +97,7 @@ bool UInteractiveToolManager::SelectActiveToolType(EToolSide Side, const FString
 
 
 
-bool UInteractiveToolManager::CanActivateTool(EToolSide Side, const FString& Identifier)
+bool UInteractiveToolManager::CanActivateTool(EToolSide Side, FString Identifier)
 {
 	check(Side == EToolSide::Left);   // TODO: support right-side tool
 
diff --git a/Engine/Source/Runtime/Experimental/InteractiveToolsFramework/Public/InteractiveToolManager.h b/Engine/Source/Runtime/Experimental/InteractiveToolsFramework/Public/InteractiveToolManager.h
--- a/Engine/Source/Runtime/Experimental/InteractiveToolsFramework/Public/InteractiveToolManager.h
+++ b/Engine/Source/Runtime/Experimental/InteractiveToolsFramework/Public/InteractiveToolManager.h
@@ -113,7 +113,7 @@ public:
 	 * @param Identifier string name of the Tool type
 	 * @return true if the Tool type could be activated
 	 */
-	virtual bool CanActivateTool(EToolSide eSide, const FString& Identifier);
+	virtual bool CanActivateTool(EToolSide eSide, FString Identifier);
 
 	/**
 	 * Try to activate a new Tool instance on the given Side
