CIS fix.

#rnx

diff --git a/Engine/Source/Editor/AssetPlacementEdMode/Private/Tools/PlacementBrushToolBase.cpp b/Engine/Source/Editor/AssetPlacementEdMode/Private/Tools/PlacementBrushToolBase.cpp
--- a/Engine/Source/Editor/AssetPlacementEdMode/Private/Tools/PlacementBrushToolBase.cpp
+++ b/Engine/Source/Editor/AssetPlacementEdMode/Private/Tools/PlacementBrushToolBase.cpp
@@ -18,7 +18,7 @@ bool UPlacementBrushToolBase::HitTest(const FRay& Ray, FHitResult& OutHit)
 	UWorld* EditingWorld = GetToolManager()->GetWorld();
 	const FVector TraceStart(Ray.Origin);
 	const FVector TraceEnd(Ray.Origin + Ray.Direction * HALF_WORLD_MAX);
-	constexpr TCHAR* NAME_PlacementBrushTool = TEXT("PlacementBrushTool");
+	constexpr TCHAR NAME_PlacementBrushTool[] = TEXT("PlacementBrushTool");
 
 	auto FilterFunc = [](const UPrimitiveComponent* InComponent) {
 		if (InComponent)
