Initial Asset Palette, selection, and single place tool functioning in placement mode.

# Removed "reapply settings" since this was a tool that was specific to UFoliageType data blueprint changes being reapplied to already placed foliage.

diff --git a/Engine/Source/Editor/AssetPlacementEdMode/AssetPlacementEdMode.Build.cs b/Engine/Source/Editor/AssetPlacementEdMode/AssetPlacementEdMode.Build.cs
--- a/Engine/Source/Editor/AssetPlacementEdMode/AssetPlacementEdMode.Build.cs
+++ b/Engine/Source/Editor/AssetPlacementEdMode/AssetPlacementEdMode.Build.cs
@@ -21,6 +21,8 @@ public class AssetPlacementEdMode : ModuleRules
 				"InteractiveToolsFramework",
 				"Foliage",
 				"Landscape",
+				"TypedElementFramework",
+				"PropertyEditor",
 			}
 		);
 	}
diff --git a/Engine/Source/Editor/AssetPlacementEdMode/Private/AssetPlacementEdMode.cpp b/Engine/Source/Editor/AssetPlacementEdMode/Private/AssetPlacementEdMode.cpp
--- a/Engine/Source/Editor/AssetPlacementEdMode/Private/AssetPlacementEdMode.cpp
+++ b/Engine/Source/Editor/AssetPlacementEdMode/Private/AssetPlacementEdMode.cpp
@@ -5,7 +5,10 @@
 #include "InteractiveToolManager.h"
 #include "AssetPlacementEdModeStyle.h"
 #include "AssetPlacementEdModeCommands.h"
+#include "AssetPlacementSettings.h"
 #include "EdModeInteractiveToolsContext.h"
+#include "EditorModeManager.h"
+#include "Selection.h"
 
 #include "Tools/PlacementSelectTool.h"
 #include "Tools/PlacementLassoSelectTool.h"
@@ -13,7 +16,10 @@
 #include "Tools/PlacementPlaceSingleTool.h"
 #include "Tools/PlacementFillTool.h"
 #include "Tools/PlacementEraseTool.h"
-#include "Tools/PlacementReapplySettingsTool.h"
+
+#include "Factories/AssetFactoryInterface.h"
+#include "Elements/Actor/ActorElementData.h"
+#include "Elements/Framework/EngineElementsLibrary.h"
 
 #include "Settings/LevelEditorMiscSettings.h"
 
@@ -37,24 +43,30 @@ UAssetPlacementEdMode::~UAssetPlacementEdMode()
 
 void UAssetPlacementEdMode::Enter()
 {
-	UEdMode::Enter();
+	Super::Enter();
 
 	const FAssetPlacementEdModeCommands& PlacementModeCommands = FAssetPlacementEdModeCommands::Get();
 	RegisterTool(PlacementModeCommands.Select, UPlacementModeSelectTool::ToolName, NewObject<UPlacementModeSelectToolBuilder>(this));
 	RegisterTool(PlacementModeCommands.LassoSelect, UPlacementModeLassoSelectTool::ToolName, NewObject<UPlacementModeLassoSelectToolBuilder>(this));
 	RegisterTool(PlacementModeCommands.Place, UPlacementModePlacementTool::ToolName, NewObject<UPlacementModePlacementToolBuilder>(this));
-	RegisterTool(PlacementModeCommands.PlaceSingle, UPlacementModePlaceSingleTool::ToolName, NewObject<UPlacementModeSelectAllToolBuilder>(this));
+
+	UPlacementModePlaceSingleToolBuilder* SinglePlaceTool = NewObject<UPlacementModePlaceSingleToolBuilder>(this);
+	SinglePlaceTool->PlacementSettings = Cast<UAssetPlacementSettings>(SettingsObject);
+	RegisterTool(PlacementModeCommands.PlaceSingle, UPlacementModePlaceSingleTool::ToolName, SinglePlaceTool);
+
 	RegisterTool(PlacementModeCommands.Fill, UPlacementModeFillTool::ToolName, NewObject<UPlacementModeFillToolBuilder>(this));
 	RegisterTool(PlacementModeCommands.Erase, UPlacementModeEraseTool::ToolName, NewObject<UPlacementModeEraseToolBuilder>(this));
-	RegisterTool(PlacementModeCommands.ReapplySettings, UPlacementModeReapplySettingsTool::ToolName, NewObject<UPlacementModeReapplySettingsToolBuilder>(this));
 
 	ToolsContext->ToolManager->SelectActiveToolType(EToolSide::Mouse, UPlacementModeSelectTool::ToolName);
 	ToolsContext->ToolManager->ActivateTool(EToolSide::Mouse);
+
+	// TODO - Update selection set based on valid things in the palette
+	Owner->GetSelectedObjects()->GetElementSelectionSet()->ClearSelection(FTypedElementSelectionOptions());
 }
 
 void UAssetPlacementEdMode::CreateToolkit()
 {
-	Toolkit = MakeShareable(new FAssetPlacementEdModeToolkit);
+	Toolkit = MakeShareable(new FAssetPlacementEdModeToolkit(Cast<UAssetPlacementSettings>(SettingsObject)));
 }
 
 TMap<FName, TArray<TSharedPtr<FUICommandInfo>>> UAssetPlacementEdMode::GetModeCommands() const
@@ -71,10 +83,6 @@ void UAssetPlacementEdMode::BindCommands()
 		FExecuteAction::CreateUObject(this, &UAssetPlacementEdMode::SelectAssets, EPaletteFilter::ActivePaletteOnly, ESelectMode::Select),
 		FCanExecuteAction::CreateUObject(this, &UAssetPlacementEdMode::HasAnyAssetsInPalette, EPaletteFilter::ActivePaletteOnly));
 
-	CommandList->MapAction(PlacementModeCommands.SelectInvalid,
-		FExecuteAction::CreateUObject(this, &UAssetPlacementEdMode::SelectAssets, EPaletteFilter::InvalidInstances, ESelectMode::Select),
-		FCanExecuteAction::CreateUObject(this, &UAssetPlacementEdMode::HasAnyAssetsInPalette, EPaletteFilter::InvalidInstances));
-
 	CommandList->MapAction(PlacementModeCommands.Deselect,
 		FExecuteAction::CreateUObject(this, &UAssetPlacementEdMode::SelectAssets, EPaletteFilter::ActivePaletteOnly, ESelectMode::Deselect),
 		FCanExecuteAction::CreateUObject(this, &UAssetPlacementEdMode::HasAnyAssetsInPalette, EPaletteFilter::ActivePaletteOnly));
@@ -82,10 +90,53 @@ void UAssetPlacementEdMode::BindCommands()
 	CommandList->MapAction(PlacementModeCommands.Delete,
 		FExecuteAction::CreateUObject(this, &UAssetPlacementEdMode::DeleteAssets),
 		FCanExecuteAction::CreateUObject(this, &UAssetPlacementEdMode::HasAnyAssetsInPalette, EPaletteFilter::ActivePaletteOnly));
+}
 
-	CommandList->MapAction(PlacementModeCommands.MoveToActivePartition,
-		FExecuteAction::CreateUObject(this, &UAssetPlacementEdMode::MoveAssetToActivePartition),
-		FCanExecuteAction::CreateUObject(this, &UAssetPlacementEdMode::HasAnyAssetsInPalette, EPaletteFilter::ActivePaletteOnly));
+bool UAssetPlacementEdMode::IsSelectionAllowed(AActor* InActor, bool bInSelection) const
+{
+	// Need to be in selection tool for selection to be allowed
+	if (GetToolManager()->GetActiveToolName(EToolSide::Mouse) != UPlacementModeSelectTool::ToolName)
+	{
+		return false;
+	}
+
+	FTypedElementHandle ActorHandle = UEngineElementsLibrary::AcquireEditorActorElementHandle(InActor);
+	if (!ActorHandle.IsSet())
+	{
+		return false;
+	}
+
+	UAssetPlacementSettings* PlacementSettings = Cast<UAssetPlacementSettings>(SettingsObject);
+	for (FPaletteItem& Item : PlacementSettings->PaletteItems)
+	{
+		FAssetData FoundAssetDataFromFactory = Item.FactoryOverride->GetAssetDataFromElementHandle(ActorHandle);
+		if (FoundAssetDataFromFactory == Item.AssetData)
+		{
+			return true;
+		}
+	}
+
+	return false;
+}
+
+bool UAssetPlacementEdMode::UsesPropertyWidgets() const
+{
+	if (GetToolManager()->GetActiveToolName(EToolSide::Mouse) != UPlacementModeSelectTool::ToolName)
+	{
+		return false;
+	}
+
+	return true;
+}
+
+bool UAssetPlacementEdMode::ShouldDrawWidget() const
+{
+	if (GetToolManager()->GetActiveToolName(EToolSide::Mouse) != UPlacementModeSelectTool::ToolName)
+	{
+		return false;
+	}
+
+	return Super::ShouldDrawWidget();
 }
 
 void UAssetPlacementEdMode::SelectAssets(EPaletteFilter InSelectAllType, ESelectMode InSelectMode)
@@ -102,7 +153,8 @@ void UAssetPlacementEdMode::MoveAssetToActivePartition()
 
 bool UAssetPlacementEdMode::HasAnyAssetsInPalette(EPaletteFilter InSelectAllType)
 {
-	return false;
+	UAssetPlacementSettings* PlacementSettings = Cast<UAssetPlacementSettings>(SettingsObject);
+	return (PlacementSettings->PaletteItems.Num() > 0);
 }
 
 #undef LOCTEXT_NAMESPACE
diff --git a/Engine/Source/Editor/AssetPlacementEdMode/Private/AssetPlacementEdModeCommands.cpp b/Engine/Source/Editor/AssetPlacementEdMode/Private/AssetPlacementEdModeCommands.cpp
--- a/Engine/Source/Editor/AssetPlacementEdMode/Private/AssetPlacementEdModeCommands.cpp
+++ b/Engine/Source/Editor/AssetPlacementEdMode/Private/AssetPlacementEdModeCommands.cpp
@@ -22,14 +22,10 @@ void FAssetPlacementEdModeCommands::RegisterCommands()
 	ToolCommands.Add(SelectAll);
 	UI_COMMAND(Deselect, "Deselect", "Deselect All Assets from the Palette", EUserInterfaceActionType::ToggleButton, FInputChord());
 	ToolCommands.Add(Deselect);
-	UI_COMMAND(SelectInvalid, "Invalid", "Select Invalid Assets", EUserInterfaceActionType::ToggleButton, FInputChord());
-	ToolCommands.Add(SelectInvalid);
 	UI_COMMAND(LassoSelect, "Lasso", "Select Assets with Lasso", EUserInterfaceActionType::ToggleButton, FInputChord());
 	ToolCommands.Add(LassoSelect);
 	UI_COMMAND(Place, "Place", "Place Assets from the Palette", EUserInterfaceActionType::ToggleButton, FInputChord());
 	ToolCommands.Add(Place);
-	UI_COMMAND(ReapplySettings, "Reapply", "Reapply Settings to Assets from the Palette", EUserInterfaceActionType::ToggleButton, FInputChord());
-	ToolCommands.Add(ReapplySettings);
 	UI_COMMAND(PlaceSingle, "Single", "Place a Single Asset from the Palette", EUserInterfaceActionType::ToggleButton, FInputChord());
 	ToolCommands.Add(PlaceSingle);
 	UI_COMMAND(Fill, "Fill", "Fill the Selected Target with Assets from the Palette", EUserInterfaceActionType::ToggleButton, FInputChord());
@@ -38,8 +34,6 @@ void FAssetPlacementEdModeCommands::RegisterCommands()
 	ToolCommands.Add(Erase);
 	UI_COMMAND(Delete, "Delete", "Remove the Selected Assets", EUserInterfaceActionType::ToggleButton, FInputChord());
 	ToolCommands.Add(Delete);
-	UI_COMMAND(MoveToActivePartition, "Move", "Move Selected Assets to the Current World Partition", EUserInterfaceActionType::ToggleButton, FInputChord());
-	ToolCommands.Add(MoveToActivePartition);
 }
 
 TMap<FName, TArray<TSharedPtr<FUICommandInfo>>> FAssetPlacementEdModeCommands::GetCommands()
diff --git a/Engine/Source/Editor/AssetPlacementEdMode/Private/AssetPlacementEdModeStyle.cpp b/Engine/Source/Editor/AssetPlacementEdMode/Private/AssetPlacementEdModeStyle.cpp
--- a/Engine/Source/Editor/AssetPlacementEdMode/Private/AssetPlacementEdModeStyle.cpp
+++ b/Engine/Source/Editor/AssetPlacementEdMode/Private/AssetPlacementEdModeStyle.cpp
@@ -5,6 +5,10 @@
 #include "Styling/SlateStyle.h"
 #include "Styling/SlateStyleMacros.h"
 #include "Styling/SlateStyleRegistry.h"
+#include "Styling/SlateTypes.h"
+#include "Styling/SlateStyleMacros.h"
+#include "Styling/StarshipCoreStyle.h"
+#include "Styling/SlateStyleRegistry.h"
 
 TSharedPtr<FAssetPlacementEdModeStyle> FAssetPlacementEdModeStyle::AssetPlacementEdModeStyle;
 
@@ -29,6 +33,12 @@ FName FAssetPlacementEdModeStyle::GetStyleSetName()
 	return "AssetPlacementEdModeStyle";
 }
 
+const ISlateStyle& FAssetPlacementEdModeStyle::Get()
+{
+	check(AssetPlacementEdModeStyle.IsValid());
+	return *AssetPlacementEdModeStyle->StyleSet;
+}
+
 FAssetPlacementEdModeStyle::FAssetPlacementEdModeStyle()
 {
 	StyleSet = MakeShareable(new FSlateStyleSet(GetStyleSetName()));
@@ -76,6 +86,13 @@ void FAssetPlacementEdModeStyle::SetupCustomStyle()
 	StyleSet->Set("AssetPlacementEdMode.MoveToActivePartition", MakeImageBrushFn("Icons/GeneralTools/MoveToLevel_40x"));
 	StyleSet->Set("AssetPlacementEdMode.MoveToActivePartition.Small", MakeImageBrushFn("Icons/GeneralTools/MoveToLevel_40x"));
 
+	StyleSet->Set("AssetPlacementEdMode.AddAssetType.Text", FTextBlockStyle(StyleSet->GetWidgetStyle<FTextBlockStyle>("NormalText"))
+		.SetFont(DEFAULT_FONT("Bold", 10))
+		.SetColorAndOpacity(FLinearColor(1.0f, 1.0f, 1.0f))
+		.SetHighlightColor(FLinearColor(1.0f, 1.0f, 1.0f))
+		.SetShadowOffset(FVector2D(1, 1))
+		.SetShadowColorAndOpacity(FLinearColor(0, 0, 0, 0.9f)));
+
 	FSlateStyleRegistry::RegisterSlateStyle(*StyleSet.Get());
 }
 
diff --git a/Engine/Source/Editor/AssetPlacementEdMode/Private/AssetPlacementEdModeStyle.h b/Engine/Source/Editor/AssetPlacementEdMode/Private/AssetPlacementEdModeStyle.h
--- a/Engine/Source/Editor/AssetPlacementEdMode/Private/AssetPlacementEdModeStyle.h
+++ b/Engine/Source/Editor/AssetPlacementEdMode/Private/AssetPlacementEdModeStyle.h
@@ -16,6 +16,7 @@ public:
 	static void Shutdown();
 
 	static FName GetStyleSetName();
+	static const ISlateStyle& Get();
 
 private:
 	FAssetPlacementEdModeStyle();
diff --git a/Engine/Source/Editor/AssetPlacementEdMode/Private/AssetPlacementEdModeToolkit.cpp b/Engine/Source/Editor/AssetPlacementEdMode/Private/AssetPlacementEdModeToolkit.cpp
--- a/Engine/Source/Editor/AssetPlacementEdMode/Private/AssetPlacementEdModeToolkit.cpp
+++ b/Engine/Source/Editor/AssetPlacementEdMode/Private/AssetPlacementEdModeToolkit.cpp
@@ -2,17 +2,19 @@
 
 #include "AssetPlacementEdModeToolkit.h"
 #include "AssetPlacementEdMode.h"
+#include "AssetPlacementSettings.h"
+#include "SAssetPlacementPalette.h"
 
 #include "Modules/ModuleManager.h"
 
-#include "Widgets/Input/SButton.h"
-#include "Widgets/Text/STextBlock.h"
+#include "Widgets/Layout/SExpandableArea.h"
 
 #include "EditorModeManager.h"
 
 #define LOCTEXT_NAMESPACE "AssetPlacementEdModeToolkit"
 
-FAssetPlacementEdModeToolkit::FAssetPlacementEdModeToolkit()
+FAssetPlacementEdModeToolkit::FAssetPlacementEdModeToolkit(TWeakObjectPtr<UAssetPlacementSettings> InPlacementSettings)
+	: PlacementSettings(InPlacementSettings)
 {
 }
 
@@ -32,4 +34,23 @@ FText FAssetPlacementEdModeToolkit::GetBaseToolkitName() const
 	return LOCTEXT("DisplayName", "AssetPlacementEdMode Tool");
 }
 
+TSharedPtr<SWidget> FAssetPlacementEdModeToolkit::GetInlineContent() const
+{
+	return SNew(SVerticalBox)
+		+ SVerticalBox::Slot()
+		[
+			SNew(SExpandableArea)
+			.AreaTitle(LOCTEXT("AssetPaletteHeader", "Asset Palette"))
+			.BodyContent()
+			[
+				SNew(SAssetPlacementPalette)
+				.PlacementSettings(PlacementSettings)
+			]
+		]
+		+ SVerticalBox::Slot()
+		[
+			FModeToolkit::GetInlineContent().ToSharedRef()
+		];
+}
+
 #undef LOCTEXT_NAMESPACE
diff --git a/Engine/Source/Editor/AssetPlacementEdMode/Private/AssetPlacementPaletteItem.cpp b/Engine/Source/Editor/AssetPlacementEdMode/Private/AssetPlacementPaletteItem.cpp
new file mode 100644
--- /dev/null
+++ b/Engine/Source/Editor/AssetPlacementEdMode/Private/AssetPlacementPaletteItem.cpp
@@ -0,0 +1,365 @@
+// Copyright Epic Games, Inc. All Rights Reserved.
+
+#include "AssetPlacementPaletteItem.h"
+#include "Widgets/Input/SButton.h"
+#include "Widgets/SToolTip.h"
+#include "Widgets/Input/SCheckBox.h"
+#include "Modules/ModuleManager.h"
+#include "Widgets/Images/SImage.h"
+#include "EditorStyleSet.h"
+
+#include "AssetThumbnail.h"
+#include "IAssetTools.h"
+#include "IAssetTypeActions.h"
+#include "AssetToolsModule.h"
+#include "FoliageType.h"
+#include "SAssetPlacementPalette.h"
+#include "Engine/Blueprint.h"
+#include "Factories/AssetFactoryInterface.h"
+#include "AssetPlacementSettings.h"
+
+#define LOCTEXT_NAMESPACE "AssetPlacementMode"
+
+FAssetPlacementPaletteItemModel::FAssetPlacementPaletteItemModel(FAssetPlacementUIInfoPtr InTypeInfo, TSharedRef<SAssetPlacementPalette> InFoliagePalette, TSharedPtr<FAssetThumbnailPool> InThumbnailPool/*, FEdModeFoliage* InFoliageEditMode*/)
+	: TypeInfo(InTypeInfo)
+	, AssetPalette(InFoliagePalette)
+{
+	check(TypeInfo);
+	DisplayFName = TypeInfo->AssetData.AssetName;
+
+	int32 MaxThumbnailSize = PlacementPaletteConstants::ThumbnailSizeRange.Max;
+	TSharedPtr<FAssetThumbnail> Thumbnail = MakeShared<FAssetThumbnail>(TypeInfo->AssetData, MaxThumbnailSize, MaxThumbnailSize, InThumbnailPool);
+	
+	FAssetThumbnailConfig ThumbnailConfig;
+	FAssetToolsModule& AssetToolsModule = FModuleManager::LoadModuleChecked<FAssetToolsModule>(TEXT("AssetTools"));
+	if (!IsBlueprint())
+	{
+		AssetTypeActions = AssetToolsModule.Get().GetAssetTypeActionsForClass(TypeInfo->AssetData.GetClass());
+	}
+	else
+	{
+		AssetTypeActions = AssetToolsModule.Get().GetAssetTypeActionsForClass(UBlueprint::StaticClass());
+	}
+
+	if (AssetTypeActions.IsValid())
+	{
+		ThumbnailConfig.AssetTypeColorOverride = AssetTypeActions.Pin()->GetTypeColor();
+	}
+
+	ThumbnailWidget = Thumbnail->MakeThumbnailWidget(ThumbnailConfig);
+}
+
+TSharedPtr<SAssetPlacementPalette> FAssetPlacementPaletteItemModel::GetAssetPalette() const
+{
+	return AssetPalette.Pin();
+}
+
+FAssetPlacementUIInfoPtr FAssetPlacementPaletteItemModel::GetTypeUIInfo() const
+{
+	return TypeInfo;
+}
+
+TSharedRef<SWidget> FAssetPlacementPaletteItemModel::GetThumbnailWidget() const
+{
+	return ThumbnailWidget.ToSharedRef();
+}
+
+TSharedRef<SToolTip> FAssetPlacementPaletteItemModel::CreateTooltipWidget() const
+{
+	return 
+		SNew(SToolTip)
+		.TextMargin(1)
+		.BorderImage(FEditorStyle::GetBrush("ContentBrowser.TileViewTooltip.ToolTipBorder"))
+		.Visibility(this, &FAssetPlacementPaletteItemModel::GetTooltipVisibility)
+		[
+			SNew(SBorder)
+			.Padding(3.f)
+			.BorderImage(FEditorStyle::GetBrush("ContentBrowser.TileViewTooltip.NonContentBorder"))
+			[
+				SNew(SVerticalBox)
+
+				+ SVerticalBox::Slot()
+				.AutoHeight()
+				[
+					SNew(SBorder)
+					.Padding(FMargin(6.f))
+					.HAlign(HAlign_Left)
+					.BorderImage(FEditorStyle::GetBrush("ContentBrowser.TileViewTooltip.ContentBorder"))
+					[
+						SNew(STextBlock)
+						.Text(FText::FromName(DisplayFName))
+						.Font(FEditorStyle::GetFontStyle("ContentBrowser.TileViewTooltip.NameFont"))
+						.HighlightText(this, &FAssetPlacementPaletteItemModel::GetPaletteSearchText)
+					]
+				]
+				
+				+ SVerticalBox::Slot()
+				.AutoHeight()
+				.Padding(FMargin(0.f, 3.f, 0.f, 0.f))
+				[
+					
+					SNew(SHorizontalBox)
+
+					+ SHorizontalBox::Slot()
+					.AutoWidth()
+					.Padding(FMargin(0.f, 0.f, 3.f, 0.f))
+					[
+						SNew(SBorder)
+						.Padding(6.f)
+						.HAlign(HAlign_Center)
+						.Visibility(this, &FAssetPlacementPaletteItemModel::GetTooltipThumbnailVisibility)
+						.BorderImage(FEditorStyle::GetBrush("ContentBrowser.TileViewTooltip.ContentBorder"))
+						[
+							SNew(SBox)
+							.HeightOverride(64.f)
+							.WidthOverride(64.f)
+							[
+								GetThumbnailWidget()
+							]
+						]
+					]
+
+					+ SHorizontalBox::Slot()
+					[
+						SNew(SBorder)
+						.Padding(6.f)
+						.BorderImage(FEditorStyle::GetBrush("ContentBrowser.TileViewTooltip.ContentBorder"))
+						[
+							SNew(SVerticalBox)
+
+							+ SVerticalBox::Slot()
+							.Padding(0, 1)
+							.AutoHeight()
+							[
+								SNew(SHorizontalBox)
+								
+								+ SHorizontalBox::Slot()
+								.AutoWidth()
+								[
+									SNew(STextBlock)
+									.Text(LOCTEXT("SourceAssetTypeHeading", "Source Asset Type: "))
+									.ColorAndOpacity(FSlateColor::UseSubduedForeground())
+								]
+
+								+ SHorizontalBox::Slot()
+								.AutoWidth()
+								[
+									SNew(STextBlock)
+									.Text(this, &FAssetPlacementPaletteItemModel::GetSourceAssetTypeText)
+								]
+							]
+						]
+					]					
+				]
+			]
+		];
+}
+
+TSharedRef<SCheckBox> FAssetPlacementPaletteItemModel::CreateActivationCheckBox(TAttribute<bool> IsItemWidgetSelected, TAttribute<EVisibility> InVisibility)
+{
+	return
+		SNew(SCheckBox)
+		.Padding(0.f)
+		.OnCheckStateChanged(this, &FAssetPlacementPaletteItemModel::HandleCheckStateChanged, IsItemWidgetSelected)
+		.Visibility(InVisibility)
+		.IsChecked(this, &FAssetPlacementPaletteItemModel::GetCheckBoxState)
+		.ToolTipText(LOCTEXT("TileCheckboxTooltip", "Check to activate the currently selected types in the palette"));
+}
+
+FName FAssetPlacementPaletteItemModel::GetDisplayFName() const
+{
+	return DisplayFName;
+}
+
+FText FAssetPlacementPaletteItemModel::GetPaletteSearchText() const
+{
+	if (AssetPalette.IsValid())
+	{
+		return AssetPalette.Pin()->GetSearchText();
+	}
+	else
+	{
+		return FText();
+	}
+}
+
+void FAssetPlacementPaletteItemModel::SetTypeActiveInPalette(bool bSetActiveInPalette)
+{
+	bSelected = bSetActiveInPalette;
+}
+
+bool FAssetPlacementPaletteItemModel::IsActive() const
+{
+	return bSelected;
+}
+
+bool FAssetPlacementPaletteItemModel::IsBlueprint() const
+{
+	return TypeInfo && (TypeInfo->AssetData.GetClass()->GetDefaultObject()->IsA<UBlueprint>() || (TypeInfo->AssetData.GetClass()->ClassGeneratedBy != nullptr));
+}
+
+bool FAssetPlacementPaletteItemModel::IsAsset() const
+{
+	return TypeInfo && TypeInfo->AssetData.IsValid();
+}
+
+void FAssetPlacementPaletteItemModel::HandleCheckStateChanged(const ECheckBoxState NewCheckedState, TAttribute<bool> IsItemWidgetSelected)
+{
+	if (!IsItemWidgetSelected.IsSet()) { return; }
+
+	const bool bShouldActivate = NewCheckedState == ECheckBoxState::Checked;
+	if (!IsItemWidgetSelected.Get())
+	{
+		SetTypeActiveInPalette(bShouldActivate);
+	}
+	else if (AssetPalette.IsValid())
+	{
+		AssetPalette.Pin()->ActivateAllSelectedTypes(bShouldActivate);
+	}
+}
+
+ECheckBoxState FAssetPlacementPaletteItemModel::GetCheckBoxState() const
+{
+	return bSelected ? ECheckBoxState::Checked : ECheckBoxState::Unchecked;
+}
+
+EVisibility FAssetPlacementPaletteItemModel::GetTooltipVisibility() const
+{
+	return (AssetPalette.IsValid() && AssetPalette.Pin()->ShouldShowTooltips()) ? EVisibility::SelfHitTestInvisible : EVisibility::Collapsed;
+}
+
+EVisibility FAssetPlacementPaletteItemModel::GetTooltipThumbnailVisibility() const
+{
+	
+	return (AssetPalette.IsValid() && AssetPalette.Pin()->IsActiveViewMode(EAssetPlacementPaletteViewMode::Tree)) ? EVisibility::SelfHitTestInvisible : EVisibility::Collapsed;
+}
+
+FText FAssetPlacementPaletteItemModel::GetSourceAssetTypeText() const
+{
+	if (AssetTypeActions.IsValid())
+	{
+		return AssetTypeActions.Pin()->GetName();
+	}
+
+	if (!TypeInfo->AssetData.AssetClass.IsNone())
+	{
+		return FText::FromName(TypeInfo->AssetData.AssetClass);
+	}
+
+	return FText::FromName(TypeInfo->AssetData.AssetName);
+}
+
+////////////////////////////////////////////////
+// SAssetPlacementPaletteItemTile
+////////////////////////////////////////////////
+
+const float SAssetPlacementPaletteItemTile::MinScaleForOverlayItems = 0.2f;
+
+void SAssetPlacementPaletteItemTile::Construct(const FArguments& InArgs, TSharedRef<STableViewBase> InOwnerTableView, TSharedPtr<FAssetPlacementPaletteItemModel>& InModel)
+{
+	Model = InModel;
+
+	auto IsSelectedGetter = TAttribute<bool>::Create(TAttribute<bool>::FGetter::CreateSP(this, &SAssetPlacementPaletteItemTile::IsSelected));
+	auto CheckBoxVisibility = TAttribute<EVisibility>::Create(TAttribute<EVisibility>::FGetter::CreateSP(this, &SAssetPlacementPaletteItemTile::GetCheckBoxVisibility));
+
+	STableRow<FAssetPlacementUIInfoPtr>::Construct(
+		STableRow<FAssetPlacementUIInfoPtr>::FArguments()
+		.Style(FEditorStyle::Get(), "ContentBrowser.AssetListView.TableRow")
+		.Padding(1.f)
+		.Content()
+		[
+			SNew(SOverlay)
+			.ToolTip(Model->CreateTooltipWidget())
+			
+			// Thumbnail
+			+ SOverlay::Slot()
+			[
+				SNew(SBorder)
+				.Padding(4.f)
+				.BorderImage(FEditorStyle::GetBrush("ContentBrowser.ThumbnailShadow"))
+				.ForegroundColor(FLinearColor::White)
+				.ColorAndOpacity(this, &SAssetPlacementPaletteItemTile::GetTileColorAndOpacity)
+				[
+					Model->GetThumbnailWidget()
+				]
+			]
+			
+			// Checkbox
+			+ SOverlay::Slot()
+			.HAlign(HAlign_Left)
+			.VAlign(VAlign_Top)
+			.Padding(FMargin(3.f))
+			[
+				SNew(SBorder)
+				.BorderImage(FEditorStyle::GetBrush("ContentBrowser.ThumbnailShadow"))
+				.BorderBackgroundColor(FLinearColor::Black)
+				.ForegroundColor(FLinearColor::White)
+				.Padding(3.f)
+				[
+					Model->CreateActivationCheckBox(IsSelectedGetter, CheckBoxVisibility)
+				]
+			]
+		], InOwnerTableView);
+}
+
+FLinearColor SAssetPlacementPaletteItemTile::GetTileColorAndOpacity() const
+{
+	return Model->IsActive() ? FLinearColor::White : FLinearColor(0.5f, 0.5f, 0.5f, 1.f);
+}
+
+EVisibility SAssetPlacementPaletteItemTile::GetCheckBoxVisibility() const
+{
+	return CanShowOverlayItems() && (IsHovered() || (IsSelected() && Model->GetAssetPalette()->AnySelectedTileHovered())) ? EVisibility::Visible : EVisibility::Collapsed;
+}
+
+bool SAssetPlacementPaletteItemTile::CanShowOverlayItems() const
+{
+	return true;
+}
+
+////////////////////////////////////////////////
+// SAssetPlacementPaletteItemRow
+////////////////////////////////////////////////
+
+void SAssetPlacementPaletteItemRow::Construct(const FArguments& InArgs, TSharedRef<STableViewBase> InOwnerTableView, TSharedPtr<FAssetPlacementPaletteItemModel>& InModel)
+{
+	Model = InModel;
+	SMultiColumnTableRow<FAssetPlacementUIInfoPtr>::Construct(FSuperRowType::FArguments(), InOwnerTableView);
+
+	SetToolTip(Model->CreateTooltipWidget());
+}
+
+TSharedRef<SWidget> SAssetPlacementPaletteItemRow::GenerateWidgetForColumn(const FName& ColumnName)
+{
+	TSharedPtr<SWidget> TableRowContent = SNullWidget::NullWidget;
+
+	if (ColumnName == AssetPlacementPaletteTreeColumns::ColumnID_ToggleActive)
+	{
+		TAttribute<bool> IsSelectedGetter = MakeAttributeSP<bool>(this, &SAssetPlacementPaletteItemTile::IsSelected);
+		TableRowContent = Model->CreateActivationCheckBox(IsSelectedGetter);
+	} 
+	else if (ColumnName == AssetPlacementPaletteTreeColumns::ColumnID_Type)
+	{
+		TableRowContent =
+			SNew(SHorizontalBox)
+			
+			+SHorizontalBox::Slot()
+			.AutoWidth()
+			[
+				SNew(SExpanderArrow, SharedThis(this))
+			]
+			
+			+SHorizontalBox::Slot()
+			.VAlign(VAlign_Center)
+			.AutoWidth()
+			[
+				SNew(STextBlock)
+				.Text(FText::FromName(Model->GetDisplayFName()))
+				.HighlightText(Model.ToSharedRef(), &FAssetPlacementPaletteItemModel::GetPaletteSearchText)
+			];
+	}
+
+	return TableRowContent.ToSharedRef();
+}
+
+#undef LOCTEXT_NAMESPACE
diff --git a/Engine/Source/Editor/AssetPlacementEdMode/Private/AssetPlacementPaletteItem.h b/Engine/Source/Editor/AssetPlacementEdMode/Private/AssetPlacementPaletteItem.h
new file mode 100644
--- /dev/null
+++ b/Engine/Source/Editor/AssetPlacementEdMode/Private/AssetPlacementPaletteItem.h
@@ -0,0 +1,128 @@
+// Copyright Epic Games, Inc. All Rights Reserved.
+
+#pragma once
+
+#include "CoreMinimal.h"
+#include "Misc/Attribute.h"
+#include "Layout/Visibility.h"
+#include "Input/Reply.h"
+#include "Widgets/DeclarativeSyntaxSupport.h"
+#include "Widgets/SWidget.h"
+#include "Widgets/Views/STableViewBase.h"
+#include "Widgets/Views/STableRow.h"
+#include "SAssetPlacementPalette.h"
+
+class FAssetThumbnailPool;
+struct FPaletteItem;
+class IAssetTypeActions;
+
+namespace PlacementPaletteConstants
+{
+	const FInt32Interval ThumbnailSizeRange(32, 128);
+}
+
+typedef TSharedPtr<FPaletteItem> FAssetPlacementUIInfoPtr;
+
+namespace AssetPlacementPaletteTreeColumns
+{
+	/** IDs for list columns */
+	static const FName ColumnID_ToggleActive("Toggle");
+	static const FName ColumnID_Type("Type");
+};
+
+class FAssetPlacementPaletteItemModel : public TSharedFromThis<FAssetPlacementPaletteItemModel>
+{
+public:
+	FAssetPlacementPaletteItemModel(FAssetPlacementUIInfoPtr InTypeInfo, TSharedRef<class SAssetPlacementPalette> InFoliagePalette, TSharedPtr<class FAssetThumbnailPool> InThumbnailPool/*, FEdModeFoliage* InFoliageEditMode*/);
+
+	/** @return The foliage palette that contains the item */
+	TSharedPtr<SAssetPlacementPalette> GetAssetPalette() const;
+
+	FAssetPlacementUIInfoPtr GetTypeUIInfo() const;
+
+	/** @return The thumbnail widget for this item */
+	TSharedRef<SWidget> GetThumbnailWidget() const;
+
+	/** @return The tooltip widget for this item */
+	TSharedRef<class SToolTip> CreateTooltipWidget() const;
+
+	/** @return The checkbox widget for activating/deactivating this type in the palette */
+	TSharedRef<class SCheckBox> CreateActivationCheckBox(TAttribute<bool> IsItemWidgetSelected, TAttribute<EVisibility> InVisibility = TAttribute<EVisibility>());
+
+	/** Gets the FName version of the displayed name of this item */
+	FName GetDisplayFName() const;
+
+	/** Gets the current search filter text */
+	FText GetPaletteSearchText() const;
+
+	/** Sets whether this foliage type is active in the palette */
+	void SetTypeActiveInPalette(bool bSetActiveInPalette);
+
+	/** Gets whether this foliage type is active in the palette */
+	bool IsActive() const;
+
+	/** @return Whether this palette item represents an instance of a foliage type blueprint class */
+	bool IsBlueprint() const;
+
+	/** @return Whether this palette item represents a foliage type asset */
+	bool IsAsset() const;
+
+private:
+	/** Handles the change in activation of the item in the palette */
+	void HandleCheckStateChanged(const ECheckBoxState NewCheckedState, TAttribute<bool> IsItemWidgetSelected);
+
+	/** Gets whether the foliage type is active in the palette  */
+	ECheckBoxState GetCheckBoxState() const;
+
+	/** Gets the visibility of the entire tooltip */
+	EVisibility GetTooltipVisibility() const;
+
+	/** Gets the visibility of the thumbnail in the tooltip */
+	EVisibility GetTooltipThumbnailVisibility() const;
+
+	/** Gets the source asset type text */
+	FText GetSourceAssetTypeText() const;
+
+	TWeakPtr<IAssetTypeActions> GetAssetTypeActions() const;
+
+private:
+	TSharedPtr<SWidget> ThumbnailWidget;
+	FName DisplayFName;
+	FAssetPlacementUIInfoPtr TypeInfo;
+	TWeakPtr<SAssetPlacementPalette> AssetPalette;
+	TWeakPtr<IAssetTypeActions> AssetTypeActions;
+	bool bSelected = true;
+};
+
+/** A tile representing a foliage type in the palette */
+class SAssetPlacementPaletteItemTile : public STableRow<FAssetPlacementUIInfoPtr>
+{
+public:
+	SLATE_BEGIN_ARGS(SAssetPlacementPaletteItemTile) {}
+	SLATE_END_ARGS()
+
+	void Construct(const FArguments& InArgs, TSharedRef<STableViewBase> InOwnerTableView, TSharedPtr<FAssetPlacementPaletteItemModel>& InModel);
+
+private:
+	FLinearColor GetTileColorAndOpacity() const;
+	EVisibility GetCheckBoxVisibility() const;
+	bool CanShowOverlayItems() const;
+
+private:
+	static const float MinScaleForOverlayItems;
+	TSharedPtr<FAssetPlacementPaletteItemModel> Model;
+};
+
+/** A tree row representing a foliage type in the palette */
+class SAssetPlacementPaletteItemRow : public SMultiColumnTableRow<FAssetPlacementUIInfoPtr>
+{
+public:
+	SLATE_BEGIN_ARGS(SAssetPlacementPaletteItemRow) {}
+	SLATE_END_ARGS()
+
+	void Construct(const FArguments& InArgs, TSharedRef<STableViewBase> InOwnerTableView, TSharedPtr<FAssetPlacementPaletteItemModel>& InModel);
+	virtual TSharedRef<SWidget> GenerateWidgetForColumn(const FName& ColumnName) override;
+
+private:
+	TSharedPtr<FAssetPlacementPaletteItemModel> Model;
+};
diff --git a/Engine/Source/Editor/AssetPlacementEdMode/Private/AssetPlacementSettings.h b/Engine/Source/Editor/AssetPlacementEdMode/Private/AssetPlacementSettings.h
new file mode 100644
--- /dev/null
+++ b/Engine/Source/Editor/AssetPlacementEdMode/Private/AssetPlacementSettings.h
@@ -0,0 +1,52 @@
+// Copyright Epic Games, Inc. All Rights Reserved.
+
+#pragma once
+
+#include "CoreMinimal.h"
+#include "AssetRegistry/AssetData.h"
+
+#include "AssetPlacementSettings.generated.h"
+
+class IAssetFactoryInterface;
+
+USTRUCT()
+struct FPaletteItem
+{
+	GENERATED_BODY()
+
+	UPROPERTY()
+	FAssetData AssetData;
+
+	UPROPERTY()
+	TScriptInterface<IAssetFactoryInterface> FactoryOverride;
+
+	UPROPERTY()
+	bool bIsEnabled = true;
+};
+
+UCLASS(config = EditorPerProjectUserSettings)
+class UAssetPlacementSettings : public UObject
+{
+	GENERATED_BODY()
+
+public:
+	UPROPERTY(config, EditAnywhere, Category = "Filters")
+	bool bLandscape = true;
+
+	UPROPERTY(config, EditAnywhere, Category = "Filters")
+	bool bStaticMeshes = true;
+
+	UPROPERTY(config, EditAnywhere, Category = "Filters")
+	bool bBSP = true;
+
+	UPROPERTY(config, EditAnywhere, Category = "Filters")
+	bool bFoliage = false;
+
+	UPROPERTY(config, EditAnywhere, Category = "Filters")
+	bool bTranslucent = false;
+
+	// todo: asset data does not serialize out correctly
+	// maybe save soft object pointers, and convert in the UI to asset data?
+	//UPROPERTY(config)
+	TArray<FPaletteItem> PaletteItems;
+};
diff --git a/Engine/Source/Editor/AssetPlacementEdMode/Private/SAssetPlacementPalette.cpp b/Engine/Source/Editor/AssetPlacementEdMode/Private/SAssetPlacementPalette.cpp
new file mode 100644
--- /dev/null
+++ b/Engine/Source/Editor/AssetPlacementEdMode/Private/SAssetPlacementPalette.cpp
@@ -0,0 +1,1111 @@
+// Copyright Epic Games, Inc. All Rights Reserved.
+
+#include "SAssetPlacementPalette.h"
+#include "Widgets/Layout/SSplitter.h"
+#include "Widgets/Input/SButton.h"
+#include "Widgets/Input/SCheckBox.h"
+#include "AssetPlacementPaletteItem.h"
+#include "Misc/MessageDialog.h"
+#include "Misc/FeedbackContext.h"
+#include "Modules/ModuleManager.h"
+#include "Framework/Application/SlateApplication.h"
+#include "Framework/Commands/UIAction.h"
+#include "Framework/Commands/UICommandList.h"
+#include "Widgets/Images/SImage.h"
+#include "Widgets/Text/SRichTextBlock.h"
+#include "Framework/MultiBox/MultiBoxBuilder.h"
+#include "Widgets/Layout/SScrollBorder.h"
+#include "Widgets/Input/SSlider.h"
+#include "EditorStyleSet.h"
+#include "AssetThumbnail.h"
+#include "PropertyEditorModule.h"
+
+#include "IContentBrowserSingleton.h"
+#include "ContentBrowserModule.h"
+#include "PropertyCustomizationHelpers.h"
+#include "IDetailsView.h"
+#include "AssetSelection.h"
+#include "ScopedTransaction.h"
+#include "AssetData.h"
+#include "Engine/StaticMesh.h"
+#include "Editor.h"
+#include "Widgets/Layout/SScrollBox.h"
+
+#include "Widgets/Layout/SScaleBox.h"
+#include "Widgets/Layout/SWidgetSwitcher.h"
+#include "Widgets/Input/SSearchBox.h"
+#include "AssetPlacementEdModeStyle.h"
+#include "AssetPlacementSettings.h"
+#include "Selection.h"
+#include "Subsystems/PlacementSubsystem.h"
+
+#define LOCTEXT_NAMESPACE "AssetPlacementMode"
+
+////////////////////////////////////////////////
+// SPlacementDragDropHandler
+////////////////////////////////////////////////
+
+/** Drag-drop zone for adding Placement types to the palette */
+class SAssetPaletteDragDropHandler : public SCompoundWidget
+{
+public:
+	SLATE_BEGIN_ARGS(SAssetPaletteDragDropHandler) {}
+		SLATE_DEFAULT_SLOT(FArguments, Content)
+		SLATE_EVENT(FOnDrop, OnDrop)
+	SLATE_END_ARGS()
+
+	void Construct(const FArguments& InArgs)
+	{
+		bIsDragOn = false;
+		OnDropDelegate = InArgs._OnDrop;
+
+		this->ChildSlot
+			[
+				SNew(SBorder)
+				.BorderImage(FAppStyle::Get().GetBrush("WhiteBrush"))
+				.BorderBackgroundColor(this, &SAssetPaletteDragDropHandler::GetBackgroundColor)
+				.Padding(FMargin(100))
+				[
+					InArgs._Content.Widget
+				]
+			];
+	}
+
+	FReply OnDrop(const FGeometry& MyGeometry, const FDragDropEvent& DragDropEvent) override
+	{
+		bIsDragOn = false;
+		if (OnDropDelegate.IsBound())
+		{
+			return OnDropDelegate.Execute(MyGeometry, DragDropEvent);
+		}
+		
+		return FReply::Handled();
+	}
+
+	virtual void OnDragEnter(const FGeometry& MyGeometry, const FDragDropEvent& DragDropEvent) override
+	{
+		bIsDragOn = true;
+	}
+
+	virtual void OnDragLeave(const FDragDropEvent& DragDropEvent) override
+	{
+		bIsDragOn = false;
+	}
+
+private:
+	FSlateColor GetBackgroundColor() const
+	{
+		return bIsDragOn ? FLinearColor(1.0f, 0.6f, 0.1f, 0.9f) : FLinearColor(0.1f, 0.1f, 0.1f, 0.9f);
+	}
+
+private:
+	FOnDrop OnDropDelegate;
+	bool bIsDragOn;
+};
+
+////////////////////////////////////////////////
+// SUneditablePlacementTypeWarning
+////////////////////////////////////////////////
+class SUneditableAssetTypeWarning : public SCompoundWidget
+{
+public:
+	SLATE_BEGIN_ARGS(SUneditableAssetTypeWarning)
+		: _WarningText()
+		, _OnHyperlinkClicked()
+	{}
+
+	/** The rich text to show in the warning */
+	SLATE_ATTRIBUTE(FText, WarningText)
+
+		/** Called when the hyperlink in the rich text is clicked */
+		SLATE_EVENT(FSlateHyperlinkRun::FOnClick, OnHyperlinkClicked)
+
+		SLATE_END_ARGS()
+
+		/** Constructs the widget */
+		void Construct(const FArguments& InArgs)
+	{
+		ChildSlot
+			[
+				SNew(SBorder)
+				.BorderImage(FAppStyle::Get().GetBrush("ToolPanel.GroupBorder"))
+				[
+					SNew(SHorizontalBox)
+					+ SHorizontalBox::Slot()
+					.AutoWidth()
+					.HAlign(HAlign_Center)
+					.VAlign(VAlign_Center)
+					.Padding(2)
+					[
+						SNew(SImage)
+						.Image(FAppStyle::Get().GetBrush("Icons.Warning"))
+					]
+					+ SHorizontalBox::Slot()
+						.VAlign(VAlign_Center)
+						.Padding(2)
+						[
+							SNew(SRichTextBlock)
+							.DecoratorStyleSet(&FAppStyle::Get())
+							.Justification(ETextJustify::Left)
+							.TextStyle(FAppStyle::Get(), "DetailsView.BPMessageTextStyle")
+							.Text(InArgs._WarningText)
+							.AutoWrapText(true)
+							+ SRichTextBlock::HyperlinkDecorator(TEXT("HyperlinkDecorator"), InArgs._OnHyperlinkClicked)
+						]
+				]
+			];
+	}
+};
+
+////////////////////////////////////////////////
+// SPlacementPalette
+////////////////////////////////////////////////
+void SAssetPlacementPalette::Construct(const FArguments& InArgs)
+{
+	bItemsNeedRebuild = false;
+	bIsRebuildTimerRegistered = false;
+	bIsRefreshTimerRegistered = false;
+	PlacementSettings = InArgs._PlacementSettings;
+	if (!PlacementSettings.IsValid())
+	{
+		PlacementSettings = GetMutableDefault<UAssetPlacementSettings>();
+	}
+
+	UICommandList = MakeShared<FUICommandList>();
+
+	// Size of the thumbnail pool should be large enough to show a reasonable amount of Placement assets on screen at once,
+	// otherwise some thumbnail images will appear duplicated.
+	ThumbnailPool = MakeShared<FAssetThumbnailPool>(64);
+
+	TypeFilter = MakeShared<PlacementTypeTextFilter>(PlacementTypeTextFilter::FItemToStringArray::CreateSP(this, &SAssetPlacementPalette::GetPaletteItemFilterString));
+
+	const FText BlankText = FText::GetEmpty();
+
+	ChildSlot
+	[
+		SNew(SVerticalBox)
+
+		+ SVerticalBox::Slot()
+		.HAlign(HAlign_Fill)
+		[
+			// Top bar
+			SNew(SBorder)
+			.BorderImage(FAppStyle::Get().GetBrush("DetailsView.CategoryTop"))
+			.BorderBackgroundColor(FLinearColor(.6f, .6f, .6f, 1.0f))
+			[
+				SNew(SVerticalBox)
+				+ SVerticalBox::Slot()
+				.VAlign(VAlign_Center)
+				[
+					SNew(SHorizontalBox)
+
+					+ SHorizontalBox::Slot()
+					.FillWidth(.75f)
+					[
+						SAssignNew(SearchBoxPtr, SSearchBox)
+						.HintText(LOCTEXT("SearchPlacementPaletteHint", "Search Palette"))
+						.OnTextChanged(this, &SAssetPlacementPalette::OnSearchTextChanged)
+					]
+
+					// View Options
+					+ SHorizontalBox::Slot()
+					.AutoWidth()
+					[
+						SNew(SComboButton)
+						.ForegroundColor(FSlateColor::UseForeground())
+						.ButtonStyle(FAppStyle::Get(), "ToggleButton")
+						.OnGetMenuContent(this, &SAssetPlacementPalette::GetViewOptionsMenuContent)
+						.ButtonContent()
+						[
+							SNew(SBox)
+							.VAlign(VAlign_Center)
+							[
+								SNew(SImage)
+								.Image(FAppStyle::Get().GetBrush("GenericViewButton"))
+							]
+						]
+					]
+				]
+
+				+ SVerticalBox::Slot()
+				.VAlign(VAlign_Center)
+				[
+					SNew(SHorizontalBox)
+					+SHorizontalBox::Slot()
+					.FillWidth(.75f)
+					[
+						// +Add Placement Type button
+						SAssignNew(AddPlacementTypeCombo, SComboButton)
+						.HAlign(HAlign_Center)
+						.ForegroundColor(FLinearColor::White)
+						.ButtonStyle(FAppStyle::Get(), "FlatButton.Success")
+						.OnGetMenuContent(this, &SAssetPlacementPalette::GetAddPlacementTypePicker)
+						.ButtonContent()
+							[
+								SNew(SHorizontalBox)
+								+ SHorizontalBox::Slot()
+								.VAlign(VAlign_Center)
+								.AutoWidth()
+								[
+									SNew(STextBlock)
+									.TextStyle(FAssetPlacementEdModeStyle::Get(), "AssetPlacementEdMode.AddAssetType.Text")
+									.Text(FText::FromString(FString(TEXT("\xf067"))) /*fa-plus*/)
+								]
+								+ SHorizontalBox::Slot()
+								.VAlign(VAlign_Center)
+								[
+									SNew(STextBlock)
+									.Text(LOCTEXT("AddPlacementTypeButtonLabel", "Add Asset Type"))
+									.TextStyle(FAssetPlacementEdModeStyle::Get(), "AssetPlacementEdMode.AddAssetType.Text")
+								]
+							]
+						]
+					]
+				]
+			]
+
+		+ SVerticalBox::Slot()
+		[
+			SNew(SOverlay)
+
+			+ SOverlay::Slot()
+			[
+				SNew(SVerticalBox)
+
+				+ SVerticalBox::Slot()
+				[
+					SNew(SBox)
+					.Visibility(this, &SAssetPlacementPalette::GetDropPlacementHintVisibility)
+					.MinDesiredHeight(100)
+					[
+						SNew(SScaleBox)
+						.Stretch(EStretch::ScaleToFit)
+						[
+							SNew(STextBlock)
+							.Text(LOCTEXT("Placement_DropStatic", "+ Drop Assets Here"))
+							.ToolTipText(LOCTEXT("Placement_DropStatic_ToolTip", "Drag and drop asset types from the Content Browser to add them to the palette."))
+						]
+					]
+				]
+
+				+ SVerticalBox::Slot()
+				.Padding(2.f, 0.f)
+				[
+					CreatePaletteViews()
+				]
+			]
+			
+			// Placement Mesh Drop Zone
+			+ SOverlay::Slot()
+			.HAlign(HAlign_Fill)
+			.VAlign(VAlign_Fill)
+			[
+				SNew(SAssetPaletteDragDropHandler)
+				.Visibility(this, &SAssetPlacementPalette::GetPlacementDropTargetVisibility)
+				.OnDrop(this, &SAssetPlacementPalette::HandlePlacementDropped)
+				[
+					SNew(SScaleBox)
+					.Stretch(EStretch::ScaleToFit)
+					[
+						SNew(STextBlock)
+						.Text(LOCTEXT("Placement_AddPlacementMesh", "+ Asset Type"))
+						.ShadowOffset(FVector2D(1.f, 1.f))
+					]
+				]
+			]
+		]
+	];
+
+	UpdatePalette(true);
+}
+
+SAssetPlacementPalette::~SAssetPlacementPalette()
+{
+}
+
+FReply SAssetPlacementPalette::OnKeyDown(const FGeometry& MyGeometry, const FKeyEvent& InKeyEvent)
+{
+	if (UICommandList->ProcessCommandBindings(InKeyEvent))
+	{
+		return FReply::Handled();
+	}
+
+	return SCompoundWidget::OnKeyDown(MyGeometry, InKeyEvent);
+}
+
+void SAssetPlacementPalette::UpdatePalette(bool bRebuildItems)
+{
+	bItemsNeedRebuild |= bRebuildItems;
+
+	if (!bIsRebuildTimerRegistered)
+	{
+		bIsRebuildTimerRegistered = true;
+		RegisterActiveTimer(0.f, FWidgetActiveTimerDelegate::CreateSP(this, &SAssetPlacementPalette::UpdatePaletteItems));
+	}
+}
+
+void SAssetPlacementPalette::RefreshPalette()
+{
+	// Do not register the refresh timer if we're pending a rebuild; rebuild should cause the palette to refresh
+	if (!bIsRefreshTimerRegistered && !bIsRebuildTimerRegistered)
+	{
+		bIsRefreshTimerRegistered = true;
+		RegisterActiveTimer(0.f, FWidgetActiveTimerDelegate::CreateSP(this, &SAssetPlacementPalette::RefreshPaletteItems));
+	}
+}
+
+bool SAssetPlacementPalette::AnySelectedTileHovered() const
+{
+	for (FPlacementPaletteItemModelPtr& PaletteItem : GetActiveViewWidget()->GetSelectedItems())
+	{
+		TSharedPtr<ITableRow> Tile = TileViewWidget->WidgetFromItem(PaletteItem);
+		if (Tile.IsValid() && Tile->AsWidget()->IsHovered())
+		{
+			return true;
+		}
+	}
+
+	return false;
+}
+
+void SAssetPlacementPalette::ActivateAllSelectedTypes(bool bActivate) const
+{
+	// Apply the new check state to all of the selected types
+	for (FPlacementPaletteItemModelPtr& PaletteItem : GetActiveViewWidget()->GetSelectedItems())
+	{
+		PaletteItem->SetTypeActiveInPalette(bActivate);
+	}
+}
+
+void SAssetPlacementPalette::RefreshActivePaletteViewWidget()
+{
+	if (ActiveViewMode == EAssetPlacementPaletteViewMode::Thumbnail)
+	{
+		TileViewWidget->RequestListRefresh();
+	}
+	else
+	{
+		TreeViewWidget->RequestTreeRefresh();
+	}
+}
+
+void SAssetPlacementPalette::AddPlacementType(const FAssetData& AssetData)
+{
+	if (AddPlacementTypeCombo.IsValid())
+	{
+		AddPlacementTypeCombo->SetIsOpen(false);
+	}
+
+	if (!AssetData.IsValid())
+	{
+		return;
+	}
+
+	if (AssetData.GetClass()->HasAnyClassFlags(CLASS_Abstract | CLASS_Deprecated | CLASS_NewerVersionExists | CLASS_NotPlaceable))
+	{
+		return;
+	}
+
+	TScriptInterface<IAssetFactoryInterface> FactoryInterface;
+	if (UPlacementSubsystem* PlacementSubystem = GEditor->GetEditorSubsystem<UPlacementSubsystem>())
+	{
+		FactoryInterface = PlacementSubystem->FindAssetFactoryFromAssetData(AssetData);
+	}
+
+	if (!FactoryInterface)
+	{
+		return;
+	}
+
+	// todo - can the palette items be a map? will this break the list view?
+	if (PaletteItems.FindByPredicate([&AssetData](TSharedPtr<FAssetPlacementPaletteItemModel>& InItem) { return InItem.IsValid() && (InItem->GetTypeUIInfo()->AssetData.ObjectPath == AssetData.ObjectPath); }))
+	{
+		return;
+	}
+
+	FAssetPlacementUIInfoPtr PlacementInfo = MakeShared<FPaletteItem>();
+	PlacementInfo->AssetData = AssetData;
+	PlacementInfo->FactoryOverride = FactoryInterface;
+	PaletteItems.Add(MakeShared<FAssetPlacementPaletteItemModel>(PlacementInfo, SharedThis(this), ThumbnailPool));
+	if (PlacementSettings.IsValid())
+	{
+		PlacementSettings->PaletteItems.Add(*PlacementInfo);
+	}
+	UpdatePalette(true);
+}
+
+TSharedRef<SWidgetSwitcher> SAssetPlacementPalette::CreatePaletteViews()
+{
+	const FText BlankText = FText::GetEmpty();
+
+	// Tile View Widget
+	SAssignNew(TileViewWidget, SPlacementTypeTileView)
+		.ListItemsSource(&FilteredItems)
+		.SelectionMode(ESelectionMode::Multi)
+		.OnGenerateTile(this, &SAssetPlacementPalette::GenerateTile)
+		.OnContextMenuOpening(this, &SAssetPlacementPalette::ConstructPlacementTypeContextMenu)
+		.OnSelectionChanged(this, &SAssetPlacementPalette::OnSelectionChanged)
+		.ItemHeight(this, &SAssetPlacementPalette::GetScaledThumbnailSize)
+		.ItemWidth(this, &SAssetPlacementPalette::GetScaledThumbnailSize)
+		.ItemAlignment(EListItemAlignment::LeftAligned)
+		.ClearSelectionOnClick(true)
+		.OnMouseButtonDoubleClick(this, &SAssetPlacementPalette::OnItemDoubleClicked);
+
+	// Tree View Widget
+	SAssignNew(TreeViewWidget, SPlacementTypeTreeView)
+		.TreeItemsSource(&FilteredItems)
+		.SelectionMode(ESelectionMode::Multi)
+		.OnGenerateRow(this, &SAssetPlacementPalette::TreeViewGenerateRow)
+		.OnContextMenuOpening(this, &SAssetPlacementPalette::ConstructPlacementTypeContextMenu)
+		.OnSelectionChanged(this, &SAssetPlacementPalette::OnSelectionChanged)
+		.OnMouseButtonDoubleClick(this, &SAssetPlacementPalette::OnItemDoubleClicked)
+		.HeaderRow
+		(
+			// Toggle Active
+			SAssignNew(TreeViewHeaderRow, SHeaderRow)
+			+ SHeaderRow::Column(AssetPlacementPaletteTreeColumns::ColumnID_ToggleActive)
+			[
+				SNew(SCheckBox)
+				.IsChecked(this, &SAssetPlacementPalette::GetState_AllMeshes)
+			.OnCheckStateChanged(this, &SAssetPlacementPalette::OnCheckStateChanged_AllMeshes)
+			]
+			.DefaultLabel(BlankText)
+			.HeaderContentPadding(FMargin(0, 1, 0, 1))
+			.HAlignHeader(HAlign_Center)
+			.HAlignCell(HAlign_Center)
+			.FixedWidth(24)
+
+			// Type
+			+ SHeaderRow::Column(AssetPlacementPaletteTreeColumns::ColumnID_Type)
+			.HeaderContentPadding(FMargin(10, 1, 0, 1))
+			.SortMode(this, &SAssetPlacementPalette::GetMeshColumnSortMode)
+			.OnSort(this, &SAssetPlacementPalette::OnTypeColumnSortModeChanged)
+			.DefaultLabel(this, &SAssetPlacementPalette::GetTypeColumnHeaderText)
+			.FillWidth(5.f)
+		);
+
+	// View Mode Switcher
+	SAssignNew(WidgetSwitcher, SWidgetSwitcher);
+
+	// Thumbnail View
+	WidgetSwitcher->AddSlot((int32)EAssetPlacementPaletteViewMode::Thumbnail)
+	[
+		SNew(SScrollBorder, TileViewWidget.ToSharedRef())
+		.Content()
+		[
+			TileViewWidget.ToSharedRef()
+		]
+	];
+
+	// Tree View
+	WidgetSwitcher->AddSlot((int32)EAssetPlacementPaletteViewMode::Tree)
+	[
+		SNew(SScrollBorder, TreeViewWidget.ToSharedRef())
+		.Style(&FAssetPlacementEdModeStyle::Get().GetWidgetStyle<FScrollBorderStyle>("FoliageEditMode.TreeView.ScrollBorder"))
+		.Content()
+		[
+			TreeViewWidget.ToSharedRef()
+		]
+	];
+
+	WidgetSwitcher->SetActiveWidgetIndex((int32)ActiveViewMode);
+
+	return WidgetSwitcher.ToSharedRef();
+}
+
+void SAssetPlacementPalette::GetPaletteItemFilterString(FPlacementPaletteItemModelPtr PaletteItemModel, TArray<FString>& OutArray) const
+{
+	OutArray.Add(PaletteItemModel->GetDisplayFName().ToString());
+}
+
+void SAssetPlacementPalette::OnSearchTextChanged(const FText& InFilterText)
+{
+	TypeFilter->SetRawFilterText(InFilterText);
+	SearchBoxPtr->SetError(TypeFilter->GetFilterErrorText());
+	UpdatePalette();
+}
+
+TSharedRef<SWidget> SAssetPlacementPalette::GetAddPlacementTypePicker()
+{
+	return PropertyCustomizationHelpers::MakeAssetPickerWithMenu(FAssetData(),
+		false,
+		TArray<const UClass*>({UObject::StaticClass()}),
+		TArray<UFactory*>(),
+		FOnShouldFilterAsset::CreateSP(this, &SAssetPlacementPalette::ShouldFilterAsset),
+		FOnAssetSelected::CreateSP(this, &SAssetPlacementPalette::AddPlacementType),
+		FSimpleDelegate());
+}
+
+bool SAssetPlacementPalette::ShouldFilterAsset(const FAssetData& InAssetData)
+{
+	UClass* Class = InAssetData.GetClass();
+
+	if (Class->HasAnyClassFlags(CLASS_Abstract | CLASS_Deprecated | CLASS_NewerVersionExists | CLASS_NotPlaceable))
+	{
+		return true;
+	}
+
+	if (UPlacementSubsystem* PlacementSubsystem = GEditor->GetEditorSubsystem<UPlacementSubsystem>())
+	{
+		return (PlacementSubsystem->FindAssetFactoryFromAssetData(InAssetData) != nullptr);
+	}
+
+	return true;
+}
+
+void SAssetPlacementPalette::SetViewMode(EAssetPlacementPaletteViewMode NewViewMode)
+{
+	if (ActiveViewMode != NewViewMode)
+	{
+		ActiveViewMode = NewViewMode;
+		switch (ActiveViewMode)
+		{
+		case EAssetPlacementPaletteViewMode::Thumbnail:
+			// Set the tile selection to be the current tree selections
+			TileViewWidget->ClearSelection();
+			for (auto& TypeInfo : TreeViewWidget->GetSelectedItems())
+			{
+				TileViewWidget->SetItemSelection(TypeInfo, true);
+			}
+			break;
+			
+		case EAssetPlacementPaletteViewMode::Tree:
+			// Set the tree selection to be the current tile selection
+			TreeViewWidget->ClearSelection();
+			for (auto& TypeInfo : TileViewWidget->GetSelectedItems())
+			{
+				TreeViewWidget->SetItemSelection(TypeInfo, true);
+			}
+			break;
+		}
+
+		WidgetSwitcher->SetActiveWidgetIndex((int32)ActiveViewMode);
+		
+		RefreshActivePaletteViewWidget();
+	}
+}
+
+bool SAssetPlacementPalette::IsActiveViewMode(EAssetPlacementPaletteViewMode ViewMode) const
+{
+	return ActiveViewMode == ViewMode;
+}
+
+void SAssetPlacementPalette::ToggleShowTooltips()
+{
+	bShowFullTooltips = !bShowFullTooltips;
+}
+
+bool SAssetPlacementPalette::ShouldShowTooltips() const
+{
+	return bShowFullTooltips;
+}
+
+FText SAssetPlacementPalette::GetSearchText() const
+{
+	return TypeFilter->GetRawFilterText();
+}
+
+void SAssetPlacementPalette::OnSelectionChanged(FPlacementPaletteItemModelPtr Item, ESelectInfo::Type SelectInfo)
+{
+	// Not yet implemented
+}
+
+void SAssetPlacementPalette::OnItemDoubleClicked(FPlacementPaletteItemModelPtr Item) const
+{
+	Item->SetTypeActiveInPalette(!Item->IsActive());
+}
+
+TSharedRef<SWidget> SAssetPlacementPalette::GetViewOptionsMenuContent()
+{
+	FMenuBuilder MenuBuilder(true, UICommandList);
+
+	MenuBuilder.BeginSection("PlacementPaletteViewMode", LOCTEXT("ViewModeHeading", "Palette View Mode"));
+	{
+		MenuBuilder.AddMenuEntry(
+			LOCTEXT("ThumbnailView", "Thumbnails"),
+			LOCTEXT("ThumbnailView_ToolTip", "Display thumbnails for each Placement type in the palette."),
+			FSlateIcon(),
+			FUIAction(
+				FExecuteAction::CreateSP(this, &SAssetPlacementPalette::SetViewMode, EAssetPlacementPaletteViewMode::Thumbnail),
+				FCanExecuteAction(),
+				FIsActionChecked::CreateSP(this, &SAssetPlacementPalette::IsActiveViewMode, EAssetPlacementPaletteViewMode::Thumbnail)
+				),
+			NAME_None,
+			EUserInterfaceActionType::ToggleButton
+			);
+
+		MenuBuilder.AddMenuEntry(
+			LOCTEXT("ListView", "List"),
+			LOCTEXT("ListView_ToolTip", "Display Placement types in the palette as a list."),
+			FSlateIcon(),
+			FUIAction(
+				FExecuteAction::CreateSP(this, &SAssetPlacementPalette::SetViewMode, EAssetPlacementPaletteViewMode::Tree),
+				FCanExecuteAction(),
+				FIsActionChecked::CreateSP(this, &SAssetPlacementPalette::IsActiveViewMode, EAssetPlacementPaletteViewMode::Tree)
+				),
+			NAME_None,
+			EUserInterfaceActionType::ToggleButton
+			);
+	}
+	MenuBuilder.EndSection();
+
+	MenuBuilder.BeginSection("PlacementPaletteViewOptions", LOCTEXT("ViewOptionsHeading", "View Options"));
+	{
+		MenuBuilder.AddMenuEntry(
+			LOCTEXT("ShowTooltips", "Show Tooltips"),
+			LOCTEXT("ShowTooltips_ToolTip", "Whether to show tooltips when hovering over Placement types in the palette."),
+			FSlateIcon(),
+			FUIAction(
+				FExecuteAction::CreateSP(this, &SAssetPlacementPalette::ToggleShowTooltips),
+				FCanExecuteAction(),
+				FIsActionChecked::CreateSP(this, &SAssetPlacementPalette::ShouldShowTooltips),
+				FIsActionButtonVisible::CreateSP(this, &SAssetPlacementPalette::IsActiveViewMode, EAssetPlacementPaletteViewMode::Tree)
+				),
+			NAME_None,
+			EUserInterfaceActionType::ToggleButton
+			);
+
+		MenuBuilder.AddWidget(
+			SNew(SHorizontalBox)
+			.Visibility(this, &SAssetPlacementPalette::GetThumbnailScaleSliderVisibility)
+			+ SHorizontalBox::Slot()
+			[
+				SNew(STextBlock)
+				.Text(LOCTEXT("ThumbnailScaleLabel", "Scale"))
+			]
+			+SHorizontalBox::Slot()
+			[
+				SNew(SSlider)
+				.ToolTipText(LOCTEXT("ThumbnailScaleToolTip", "Adjust the size of thumbnails."))
+				.Value(this, &SAssetPlacementPalette::GetThumbnailScale)
+				.OnValueChanged(this, &SAssetPlacementPalette::SetThumbnailScale)
+				.OnMouseCaptureEnd(this, &SAssetPlacementPalette::RefreshActivePaletteViewWidget)
+			],
+			FText(),
+			/*bNoIndent=*/true
+			);
+	}
+	MenuBuilder.EndSection();
+
+	return MenuBuilder.MakeWidget();
+}
+
+TSharedPtr<SListView<FPlacementPaletteItemModelPtr>> SAssetPlacementPalette::GetActiveViewWidget() const
+{
+	if (ActiveViewMode == EAssetPlacementPaletteViewMode::Thumbnail)
+	{
+		return TileViewWidget;
+	}
+	else if (ActiveViewMode == EAssetPlacementPaletteViewMode::Tree)
+	{
+		return TreeViewWidget;
+	}
+	
+	return nullptr;
+}
+
+EVisibility SAssetPlacementPalette::GetDropPlacementHintVisibility() const
+{
+	return FilteredItems.Num() ? EVisibility::Collapsed : EVisibility::Visible;
+}
+
+EVisibility SAssetPlacementPalette::GetPlacementDropTargetVisibility() const
+{
+	if (FSlateApplication::Get().IsDragDropping())
+	{
+		TArray<FAssetData> DraggedAssets = AssetUtil::ExtractAssetDataFromDrag(FSlateApplication::Get().GetDragDroppingContent());
+		for (const FAssetData& AssetData : DraggedAssets)
+		{
+			if (AssetData.GetClass()->HasAnyClassFlags(CLASS_Abstract | CLASS_Deprecated | CLASS_NewerVersionExists | CLASS_NotPlaceable))
+			{
+				continue;
+			}
+
+			if (AssetData.IsValid())
+			{
+				return EVisibility::Visible;
+			}
+		}
+	}
+
+	return EVisibility::Hidden;
+}
+
+FReply SAssetPlacementPalette::HandlePlacementDropped(const FGeometry& DropZoneGeometry, const FDragDropEvent& DragDropEvent)
+{
+	TArray<FAssetData> DroppedAssetData = AssetUtil::ExtractAssetDataFromDrag(DragDropEvent);
+	if (DroppedAssetData.Num() > 0)
+	{
+		// Treat the entire drop as a transaction (in case multiples types are being added)
+		const FScopedTransaction Transaction(NSLOCTEXT("UnrealEd", "PlacementMode_DragDropTypesTransaction", "Drag-drop Placement"));
+
+		for (auto& AssetData : DroppedAssetData)
+		{
+			AddPlacementType(AssetData);
+		}
+	}
+
+	return FReply::Handled();
+}
+
+TSharedPtr<SWidget> SAssetPlacementPalette::ConstructPlacementTypeContextMenu()
+{
+	FMenuBuilder MenuBuilder(true, nullptr);
+	return MenuBuilder.MakeWidget();
+}
+
+void SAssetPlacementPalette::OnActivatePlacementTypes()
+{
+	for (FPlacementPaletteItemModelPtr& PaletteItem : GetActiveViewWidget()->GetSelectedItems())
+	{
+		PaletteItem->SetTypeActiveInPalette(true);
+	}
+}
+
+bool SAssetPlacementPalette::OnCanActivatePlacementTypes() const
+{
+	// At least one selected item must be inactive
+	for (FPlacementPaletteItemModelPtr& PaletteItem : GetActiveViewWidget()->GetSelectedItems())
+	{
+		if (!PaletteItem->IsActive())
+		{
+			return true;
+		}
+	}
+
+	return false;
+}
+
+void SAssetPlacementPalette::OnDeactivatePlacementTypes()
+{
+	for (FPlacementPaletteItemModelPtr& PaletteItem : GetActiveViewWidget()->GetSelectedItems())
+	{
+		PaletteItem->SetTypeActiveInPalette(false);
+	}
+}
+
+bool SAssetPlacementPalette::OnCanDeactivatePlacementTypes() const
+{
+	// At least one selected item must be active
+	for (FPlacementPaletteItemModelPtr& PaletteItem : GetActiveViewWidget()->GetSelectedItems())
+	{
+		if (PaletteItem->IsActive())
+		{
+			return true;
+		}
+	}
+
+	return false;
+}
+
+void SAssetPlacementPalette::FillReplacePlacementTypeSubmenu(FMenuBuilder& MenuBuilder)
+{
+	FContentBrowserModule& ContentBrowserModule = FModuleManager::Get().LoadModuleChecked<FContentBrowserModule>(TEXT("ContentBrowser"));
+
+	FAssetPickerConfig AssetPickerConfig;
+	AssetPickerConfig.Filter.bRecursiveClasses = true;
+	AssetPickerConfig.OnAssetSelected = FOnAssetSelected::CreateSP(this, &SAssetPlacementPalette::OnReplacePlacementTypeSelected);
+	AssetPickerConfig.InitialAssetViewType = EAssetViewType::List;
+	AssetPickerConfig.bAllowNullSelection = false;
+
+	TSharedRef<SWidget> MenuContent = SNew(SBox)
+		.WidthOverride(384)
+		.HeightOverride(500)
+		[
+			ContentBrowserModule.Get().CreateAssetPicker(AssetPickerConfig)
+		];
+
+	MenuBuilder.AddWidget(MenuContent, FText::GetEmpty(), true);
+}
+
+void SAssetPlacementPalette::OnReplacePlacementTypeSelected(const FAssetData& AssetData)
+{
+	FSlateApplication::Get().DismissAllMenus();
+
+	// not yet implemented
+	// should replace existing matching elements in level with the newly selected asset data
+	// and update the palette
+}
+
+void SAssetPlacementPalette::OnRemovePlacementType()
+{
+	TArray<FPlacementPaletteItemModelPtr> PaletteCopy = PaletteItems;
+	for (FPlacementPaletteItemModelPtr& PaletteItem : GetActiveViewWidget()->GetSelectedItems())
+	{
+		PaletteCopy.Remove(PaletteItem);
+	}
+
+	PaletteItems = MoveTemp(PaletteCopy);
+	GetActiveViewWidget()->RequestListRefresh();
+}
+
+void SAssetPlacementPalette::OnShowPlacementTypeInCB()
+{
+	TArray<FAssetData> SelectedAssets;
+	for (FPlacementPaletteItemModelPtr& PaletteItem : GetActiveViewWidget()->GetSelectedItems())
+	{
+		if (PaletteItem.IsValid() && PaletteItem->GetTypeUIInfo().IsValid())
+		{
+			SelectedAssets.Add(PaletteItem->GetTypeUIInfo()->AssetData);
+		}
+	}
+
+	if (SelectedAssets.Num())
+	{
+		GEditor->SyncBrowserToObjects(SelectedAssets);
+	}
+}
+
+void SAssetPlacementPalette::OnReflectSelectionInPalette()
+{
+	TArray<FAssetData> SelectedPlacementTypes;
+	for (UObject* SelectedObject : GEditor->GetSelectedObjects()->GetElementSelectionSet()->GetSelectedObjects())
+	{
+		SelectedPlacementTypes.Add(FAssetData(SelectedObject));
+	}
+	SelectPlacementTypesInPalette(SelectedPlacementTypes);
+}
+
+void SAssetPlacementPalette::SelectPlacementTypesInPalette(const TArray<FAssetData>& PlacementTypes)
+{
+	TArray<FPlacementPaletteItemModelPtr> SelectedItems;
+	SelectedItems.Reserve(PlacementTypes.Num());
+
+	for (FPlacementPaletteItemModelPtr& PaletteItem : FilteredItems)
+	{
+		if (PlacementTypes.Contains(PaletteItem->GetTypeUIInfo()->AssetData))
+		{
+			SelectedItems.Add(PaletteItem);
+		}
+	}
+	
+	GetActiveViewWidget()->ClearSelection();
+	GetActiveViewWidget()->SetItemSelection(SelectedItems, true);
+}
+
+void SAssetPlacementPalette::ExecuteOnSelectedItemPlacementTypes(TFunctionRef<void(const TArray<FAssetData>&)> ExecuteFunc)
+{
+	TArray<FPlacementPaletteItemModelPtr> SelectedItems = GetActiveViewWidget()->GetSelectedItems();
+	TArray<FAssetData> PlacementTypes;
+	PlacementTypes.Reserve(SelectedItems.Num());
+	for (FPlacementPaletteItemModelPtr& PaletteItem : SelectedItems)
+	{
+		PlacementTypes.Add(PaletteItem->GetTypeUIInfo()->AssetData);
+	}
+	ExecuteFunc(PlacementTypes);
+}
+
+void SAssetPlacementPalette::OnSelectAllInstances()
+{
+	ExecuteOnSelectedItemPlacementTypes([&](const TArray<FAssetData>& PlacementTypes)
+	{
+		//PlacementEditMode->SelectInstances(PlacementTypes, true);
+	});	
+}
+
+void SAssetPlacementPalette::OnDeselectAllInstances()
+{
+	ExecuteOnSelectedItemPlacementTypes([&](const TArray<FAssetData>& PlacementTypes)
+	{
+		//PlacementEditMode->SelectInstances(PlacementTypes, false);
+	});
+}
+
+void SAssetPlacementPalette::OnSelectInvalidInstances()
+{
+	ExecuteOnSelectedItemPlacementTypes([&](const TArray<FAssetData>& PlacementTypes)
+	{
+		//PlacementEditMode->SelectInvalidInstances(PlacementTypes);
+	});
+}
+
+bool SAssetPlacementPalette::CanSelectInstances() const
+{
+	return false;
+	//return PlacementEditMode->UISettings.GetSelectToolSelected() || PlacementEditMode->UISettings.GetLassoSelectToolSelected();
+}
+
+// THUMBNAIL VIEW
+
+TSharedRef<ITableRow> SAssetPlacementPalette::GenerateTile(FPlacementPaletteItemModelPtr Item, const TSharedRef<STableViewBase>& OwnerTable)
+{
+	return SNew(SAssetPlacementPaletteItemTile, OwnerTable, Item);
+
+	// Refresh the palette to ensure that thumbnails are correct
+	RefreshPalette();
+}
+
+float SAssetPlacementPalette::GetScaledThumbnailSize() const
+{
+	const FInt32Interval& SizeRange = PlacementPaletteConstants::ThumbnailSizeRange;
+	return SizeRange.Min + SizeRange.Size() * GetThumbnailScale();// PlacementEditMode->UISettings.GetPaletteThumbnailScale();
+}
+
+float SAssetPlacementPalette::GetThumbnailScale() const
+{
+	return PaletteThumbnailScale;
+}
+
+void SAssetPlacementPalette::SetThumbnailScale(float InScale)
+{
+	PaletteThumbnailScale = FMath::Clamp(InScale, 0.f, 1.f);
+}
+
+EVisibility SAssetPlacementPalette::GetThumbnailScaleSliderVisibility() const
+{
+	return (ActiveViewMode == EAssetPlacementPaletteViewMode::Thumbnail) ? EVisibility::Visible : EVisibility::Collapsed;
+}
+
+// TREE VIEW
+
+TSharedRef<ITableRow> SAssetPlacementPalette::TreeViewGenerateRow(FPlacementPaletteItemModelPtr Item, const TSharedRef<STableViewBase>& OwnerTable)
+{
+	return SNew(SAssetPlacementPaletteItemRow, OwnerTable, Item);
+}
+
+ECheckBoxState SAssetPlacementPalette::GetState_AllMeshes() const
+{
+	bool bHasChecked = false;
+	bool bHasUnchecked = false;
+
+	for (const FPlacementPaletteItemModelPtr& PaletteItem : FilteredItems)
+	{
+		if (PaletteItem->IsActive())
+		{
+			bHasChecked = true;
+		}
+		else
+		{
+			bHasUnchecked = true;
+		}
+
+		if (bHasChecked && bHasUnchecked)
+		{
+			return ECheckBoxState::Undetermined;
+		}
+	}
+
+	return bHasChecked ? ECheckBoxState::Checked : ECheckBoxState::Unchecked;
+}
+
+void SAssetPlacementPalette::OnCheckStateChanged_AllMeshes(ECheckBoxState InState)
+{
+	const bool bActivate = InState == ECheckBoxState::Checked;
+	for (FPlacementPaletteItemModelPtr& PaletteItem : FilteredItems)
+	{
+		PaletteItem->SetTypeActiveInPalette(bActivate);
+	}
+}
+
+FText SAssetPlacementPalette::GetTypeColumnHeaderText() const
+{
+	return LOCTEXT("PlacementTypeHeader", "Asset Type");
+}
+
+EColumnSortMode::Type SAssetPlacementPalette::GetMeshColumnSortMode() const
+{
+	return ActiveSortOrder;
+}
+
+void SAssetPlacementPalette::OnTypeColumnSortModeChanged(EColumnSortPriority::Type InPriority, const FName& InColumnName, EColumnSortMode::Type InSortMode)
+{
+	if (ActiveSortOrder == InSortMode)
+	{
+		return;
+	}
+
+	ActiveSortOrder = InSortMode;
+
+	if (ActiveSortOrder != EColumnSortMode::None)
+	{
+		auto CompareEntry = [this](const FPlacementPaletteItemModelPtr& A, const FPlacementPaletteItemModelPtr& B)
+		{
+			bool CompareResult = (A->GetDisplayFName().GetComparisonIndex().CompareLexical(B->GetDisplayFName().GetComparisonIndex()) <= 0);
+			return (ActiveSortOrder == EColumnSortMode::Ascending) ? CompareResult : !CompareResult;
+		};
+
+		PaletteItems.Sort(CompareEntry);
+	}
+}
+
+void SAssetPlacementPalette::OnEditPlacementTypeBlueprintHyperlinkClicked(const FSlateHyperlinkRun::FMetadata& Metadata)
+{
+	UBlueprint* Blueprint = nullptr; 
+	
+	// Get the first selected Placement type blueprint
+	for (FPlacementPaletteItemModelPtr& PaletteItem : GetActiveViewWidget()->GetSelectedItems())
+	{
+		Blueprint = Cast<UBlueprint>(PaletteItem->GetTypeUIInfo()->AssetData.GetClass()->ClassGeneratedBy);
+		if (Blueprint != nullptr)
+		{
+			break;
+		}
+	}
+
+	if (Blueprint)
+	{
+		// Open the blueprint
+		GEditor->EditObject(Blueprint);
+	}
+}
+
+EActiveTimerReturnType SAssetPlacementPalette::UpdatePaletteItems(double InCurrentTime, float InDeltaTime)
+{
+	if (bItemsNeedRebuild)
+	{
+		bItemsNeedRebuild = false;
+
+		// Cache the currently selected items
+		auto ActiveViewWidget = GetActiveViewWidget();
+		TArray<FPlacementPaletteItemModelPtr> PreviouslySelectedItems = ActiveViewWidget->GetSelectedItems();
+		ActiveViewWidget->ClearSelection();
+
+		// Restore the selection
+		for (auto& PrevSelectedItem : PreviouslySelectedItems)
+		{
+			// Select any replacements for previously selected Placement types
+			for (auto& Item : PaletteItems)
+			{
+				if (Item->GetDisplayFName() == PrevSelectedItem->GetDisplayFName())
+				{
+					ActiveViewWidget->SetItemSelection(Item, true);
+					break;
+				}
+			}
+		}
+	}
+
+	// Update the filtered items
+	FilteredItems.Empty();
+	for (auto& Item : PaletteItems)
+	{
+		if (TypeFilter->PassesFilter(Item))
+		{
+			FilteredItems.Add(Item);
+		}
+
+		FPaletteItem NewItem;
+		NewItem.AssetData = Item->GetTypeUIInfo()->AssetData;
+		NewItem.bIsEnabled = false;
+	}
+
+	// Refresh the appropriate view
+	RefreshActivePaletteViewWidget();
+
+	bIsRebuildTimerRegistered = false;
+	return EActiveTimerReturnType::Stop;
+}
+
+EActiveTimerReturnType SAssetPlacementPalette::RefreshPaletteItems(double InCurrentTime, float InDeltaTime)
+{
+	// Do not refresh the palette if we're waiting on a rebuild
+	if (!bItemsNeedRebuild)
+	{
+		RefreshActivePaletteViewWidget();
+	}
+
+	bIsRefreshTimerRegistered = false;
+	return EActiveTimerReturnType::Stop;
+}
+
+#undef LOCTEXT_NAMESPACE
diff --git a/Engine/Source/Editor/AssetPlacementEdMode/Private/SAssetPlacementPalette.h b/Engine/Source/Editor/AssetPlacementEdMode/Private/SAssetPlacementPalette.h
new file mode 100644
--- /dev/null
+++ b/Engine/Source/Editor/AssetPlacementEdMode/Private/SAssetPlacementPalette.h
@@ -0,0 +1,269 @@
+// Copyright Epic Games, Inc. All Rights Reserved.
+
+#pragma once
+
+#include "CoreMinimal.h"
+#include "Layout/Visibility.h"
+#include "Input/Reply.h"
+#include "Widgets/DeclarativeSyntaxSupport.h"
+#include "Widgets/SWidget.h"
+#include "Widgets/SCompoundWidget.h"
+#include "Framework/Text/SlateHyperlinkRun.h"
+#include "Widgets/Input/SComboButton.h"
+#include "Widgets/Views/SHeaderRow.h"
+#include "Widgets/Views/STableViewBase.h"
+#include "Widgets/Views/STableRow.h"
+#include "Widgets/Views/STileView.h"
+#include "Widgets/Views/STreeView.h"
+#include "Misc/TextFilter.h"
+
+class FAssetThumbnailPool;
+class FAssetPlacementPaletteItemModel;
+class FMenuBuilder;
+class FUICommandList;
+class IDetailsView;
+class UFoliageType;
+struct FAssetData;
+class UAssetPlacementSettings;
+class IPropertyHandle;
+
+typedef TSharedPtr<FAssetPlacementPaletteItemModel> FPlacementPaletteItemModelPtr;
+typedef STreeView<FPlacementPaletteItemModelPtr> SPlacementTypeTreeView;
+typedef STileView<FPlacementPaletteItemModelPtr> SPlacementTypeTileView;
+
+/** View modes supported by the palette */
+enum class EAssetPlacementPaletteViewMode : int32
+{
+	Thumbnail,
+	Tree
+};
+
+/** The palette of Placement types available for use by the Placement edit mode */
+class SAssetPlacementPalette : public SCompoundWidget
+{
+public:
+	SLATE_BEGIN_ARGS(SAssetPlacementPalette) {}
+		SLATE_ARGUMENT(TWeakObjectPtr<UAssetPlacementSettings>, PlacementSettings)
+		SLATE_ARGUMENT(TSharedPtr<IPropertyHandle>, PalettePropertyHandle)
+	SLATE_END_ARGS()
+
+	void Construct(const FArguments& InArgs);
+	~SAssetPlacementPalette();
+
+	// SWidget interface
+	virtual FReply OnKeyDown(const FGeometry& MyGeometry, const FKeyEvent& InKeyEvent) override;
+
+	/** Updates the Placement palette, optionally doing a full rebuild of the items in the palette as well */
+	void UpdatePalette(bool bRebuildItems = false);
+
+	/** Refreshes the Placement palette */
+	void RefreshPalette();
+
+	bool AnySelectedTileHovered() const;
+	void ActivateAllSelectedTypes(bool bActivate) const;
+
+	/** @return True if the given view mode is the active view mode */
+	bool IsActiveViewMode(EAssetPlacementPaletteViewMode ViewMode) const;
+
+	/** @return True if tooltips should be shown when hovering over Placement type items in the palette */
+	bool ShouldShowTooltips() const;
+
+	/** @return The current search filter text */
+	FText GetSearchText() const;
+
+	/** Adds the Placement type asset to the instanced Placement actor's list of types. */
+	void AddPlacementType(const FAssetData& AssetData);
+
+private:	// GENERAL
+
+	/** Refreshes the active palette view widget */
+	void RefreshActivePaletteViewWidget();
+
+	/** Creates the palette views */
+	TSharedRef<class SWidgetSwitcher> CreatePaletteViews();
+
+	/** Adds the displayed name of the Placement type for filtering */
+	void GetPaletteItemFilterString(FPlacementPaletteItemModelPtr PaletteItemModel, TArray<FString>& OutArray) const;
+
+	/** Handles changes to the search filter text */
+	void OnSearchTextChanged(const FText& InFilterText);
+
+	/** Gets the asset picker for adding a Placement type. */
+	TSharedRef<SWidget> GetAddPlacementTypePicker();
+
+	bool ShouldFilterAsset(const FAssetData& InAssetData);
+
+	/** Gets the visibility of the Add Placement Type text in the header row button */
+	EVisibility GetAddPlacementTypeButtonTextVisibility() const;
+
+	/** Toggle whether all Placement types are active */
+	ECheckBoxState GetState_AllMeshes() const;
+	void OnCheckStateChanged_AllMeshes(ECheckBoxState InState);
+
+	/** Sets the view mode of the palette */
+	void SetViewMode(EAssetPlacementPaletteViewMode NewViewMode);
+
+	/** Sets whether to show tooltips when hovering over Placement type items in the palette */
+	void ToggleShowTooltips();
+
+	/** Switches the palette display between the tile and tree view */
+	FReply OnToggleViewModeClicked();
+
+	/** @return The index of the view widget to display */
+	int32 GetActiveViewIndex() const;
+
+	/** Handler for selection changes in either view */
+	void OnSelectionChanged(FPlacementPaletteItemModelPtr Item, ESelectInfo::Type SelectInfo);
+
+	/** Toggle the activation state of a type on a double-click */
+	void OnItemDoubleClicked(FPlacementPaletteItemModelPtr Item) const;
+
+	/** Creates the view options menu */
+	TSharedRef<SWidget> GetViewOptionsMenuContent();
+
+	TSharedPtr<SListView<FPlacementPaletteItemModelPtr>> GetActiveViewWidget() const;
+
+	/** Gets the visibility of the "Drop Placement Here" prompt for when the palette is empty */
+	EVisibility GetDropPlacementHintVisibility() const;
+
+	/** Gets the visibility of the drag-drop zone overlay */
+	EVisibility GetPlacementDropTargetVisibility() const;
+
+	/** Handles dropping of a mesh or Placement type into the palette */
+	FReply HandlePlacementDropped(const FGeometry& DropZoneGeometry, const FDragDropEvent& DragDropEvent);
+
+private:	// CONTEXT MENU
+
+	/** @return the SWidget containing the context menu */
+	TSharedPtr<SWidget> ConstructPlacementTypeContextMenu();
+
+	/** Handler for the 'Activate' command */
+	void OnActivatePlacementTypes();
+	bool OnCanActivatePlacementTypes() const;
+
+	/** Handler for the 'Deactivate' command */
+	void OnDeactivatePlacementTypes();
+	bool OnCanDeactivatePlacementTypes() const;
+
+	/** Fills 'Replace' menu command  */
+	void FillReplacePlacementTypeSubmenu(FMenuBuilder& MenuBuilder);
+
+	/** Handler for 'Replace' command  */
+	void OnReplacePlacementTypeSelected(const struct FAssetData& AssetData);
+
+	/** Handler for 'Remove' command  */
+	void OnRemovePlacementType();
+
+	/** Handler for 'Show in CB' command  */
+	void OnShowPlacementTypeInCB();
+
+	/** Handler for 'Select All' command  */
+	void OnSelectAllInstances();
+
+	/** Handler for 'Deselect All' command  */
+	void OnDeselectAllInstances();
+
+	/** Handler for 'Select Invalid Instances' command  */
+	void OnSelectInvalidInstances();
+
+	/** Executes Function on gathered list of Placement types from the selected palette items */
+	void ExecuteOnSelectedItemPlacementTypes(TFunctionRef<void(const TArray<FAssetData>&)> ExecuteFunc);
+
+	/** @return Whether selecting instances is currently possible */
+	bool CanSelectInstances() const;
+
+	/** Handler for 'Reflect Selection in Palette ' command */
+	void OnReflectSelectionInPalette();
+
+	/** Selects Placement Type in palette */
+	void SelectPlacementTypesInPalette(const TArray<FAssetData>& PlacementTypes);
+
+private:	// THUMBNAIL VIEW
+
+	/** Creates a thumbnail tile for the given Placement type */
+	TSharedRef<ITableRow> GenerateTile(FPlacementPaletteItemModelPtr Item, const TSharedRef<STableViewBase>& OwnerTable);
+
+	/** Gets the scaled thumbnail tile size */
+	float GetScaledThumbnailSize() const;
+
+	/** Gets the current scale of the thumbnail tiles */
+	float GetThumbnailScale() const;
+
+	/** Sets the current scale of the thumbnail tiles */
+	void SetThumbnailScale(float InScale);
+
+	/** Gets whether the thumbnail scaling slider is visible */
+	EVisibility GetThumbnailScaleSliderVisibility() const;
+
+private:	// TREE VIEW
+
+	/** Generates a row widget for Placement mesh item */
+	TSharedRef<ITableRow> TreeViewGenerateRow(FPlacementPaletteItemModelPtr Item, const TSharedRef<STableViewBase>& OwnerTable);
+
+	/** Text for Placement meshes list header */
+	FText GetTypeColumnHeaderText() const;
+
+	/** Mesh list sorting support */
+	EColumnSortMode::Type GetMeshColumnSortMode() const;
+	void OnTypeColumnSortModeChanged(EColumnSortPriority::Type InPriority, const FName& InColumnName, EColumnSortMode::Type InSortMode);
+
+private:
+	/** Handles the click for the uneditable blueprint Placement type warning */
+	void OnEditPlacementTypeBlueprintHyperlinkClicked(const FSlateHyperlinkRun::FMetadata& Metadata);
+
+private:
+	/** Active timer handler to update the items in the palette */
+	EActiveTimerReturnType UpdatePaletteItems(double InCurrentTime, float InDeltaTime);
+
+	/** Active timer handler to refresh the palette */
+	EActiveTimerReturnType RefreshPaletteItems(double InCurrentTime, float InDeltaTime);
+
+private:
+	typedef TTextFilter<FPlacementPaletteItemModelPtr> PlacementTypeTextFilter;
+	TSharedPtr<PlacementTypeTextFilter> TypeFilter;
+
+	/** All the items in the palette (unfiltered) */
+	TArray<FPlacementPaletteItemModelPtr> PaletteItems;
+
+	/** The filtered list of types to display in the palette */
+	TArray<FPlacementPaletteItemModelPtr> FilteredItems;
+
+	/** Switches between the thumbnail and tree views */
+	TSharedPtr<class SWidgetSwitcher> WidgetSwitcher;
+
+	/** The Add Placement Type combo button */
+	TSharedPtr<class SComboButton> AddPlacementTypeCombo;
+
+	/** The header row of the Placement mesh tree */
+	TSharedPtr<class SHeaderRow> TreeViewHeaderRow;
+
+	/** Placement type thumbnails widget  */
+	TSharedPtr<SPlacementTypeTileView> TileViewWidget;
+
+	/** Placement type tree widget  */
+	TSharedPtr<SPlacementTypeTreeView> TreeViewWidget;
+
+	/** Placement mesh details widget  */
+	TSharedPtr<class IDetailsView> DetailsWidget;
+
+	TWeakObjectPtr<UAssetPlacementSettings> PlacementSettings;
+	TSharedPtr<IPropertyHandle> PalettePropertyHandle;
+
+	/** Placement items search box widget */
+	TSharedPtr<class SSearchBox> SearchBoxPtr;
+
+	/** Command list for binding functions for the context menu. */
+	TSharedPtr<FUICommandList> UICommandList;
+
+	/** Thumbnail pool for rendering mesh thumbnails */
+	TSharedPtr<class FAssetThumbnailPool> ThumbnailPool;
+
+	bool bItemsNeedRebuild = true;
+	bool bShowFullTooltips = true;
+	bool bIsRebuildTimerRegistered = true;
+	bool bIsRefreshTimerRegistered = true;
+	EAssetPlacementPaletteViewMode ActiveViewMode;
+	EColumnSortMode::Type ActiveSortOrder = EColumnSortMode::Type::Ascending;
+
+	float PaletteThumbnailScale = .3f;
+};
diff --git a/Engine/Source/Editor/AssetPlacementEdMode/Private/Tools/PlacementBrushToolBase.cpp b/Engine/Source/Editor/AssetPlacementEdMode/Private/Tools/PlacementBrushToolBase.cpp
--- a/Engine/Source/Editor/AssetPlacementEdMode/Private/Tools/PlacementBrushToolBase.cpp
+++ b/Engine/Source/Editor/AssetPlacementEdMode/Private/Tools/PlacementBrushToolBase.cpp
@@ -12,6 +12,20 @@
 #include "Components/ModelComponent.h"
 #include "LandscapeHeightfieldCollisionComponent.h"
 #include "FoliageInstancedStaticMeshComponent.h"
+#include "AssetPlacementSettings.h"
+
+bool UPlacementToolBuilderBase::CanBuildTool(const FToolBuilderState& SceneState) const
+{
+	return PlacementSettings && PlacementSettings->PaletteItems.Num();
+}
+
+UInteractiveTool* UPlacementToolBuilderBase::BuildTool(const FToolBuilderState& SceneState) const
+{
+	UPlacementBrushToolBase* NewTool = FactoryToolInstance(SceneState.ToolManager);
+	NewTool->PlacementSettings = PlacementSettings;
+
+	return NewTool;
+}
 
 bool UPlacementBrushToolBase::HitTest(const FRay& Ray, FHitResult& OutHit)
 {
@@ -20,15 +34,15 @@ bool UPlacementBrushToolBase::HitTest(const FRay& Ray, FHitResult& OutHit)
 	const FVector TraceEnd(Ray.Origin + Ray.Direction * HALF_WORLD_MAX);
 	constexpr TCHAR NAME_PlacementBrushTool[] = TEXT("PlacementBrushTool");
 
-	auto FilterFunc = [](const UPrimitiveComponent* InComponent) {
-		if (InComponent)
+	auto FilterFunc = [this](const UPrimitiveComponent* InComponent) {
+		if (InComponent && this->PlacementSettings.IsValid())
 		{
 			bool bFoliageOwned = InComponent->GetOwner() && FFoliageHelper::IsOwnedByFoliage(InComponent->GetOwner());
-			constexpr bool bAllowLandscape = true;
-			constexpr bool bAllowStaticMesh = true;
-			constexpr bool bAllowBSP = true;
-			constexpr bool bAllowFoliage = true;
-			constexpr bool bAllowTranslucent = true;
+			const bool bAllowLandscape = this->PlacementSettings->bLandscape;
+			const bool bAllowStaticMesh = this->PlacementSettings->bStaticMeshes;
+			const bool bAllowBSP = this->PlacementSettings->bBSP;
+			const bool bAllowFoliage = this->PlacementSettings->bFoliage;
+			const bool bAllowTranslucent = this->PlacementSettings->bTranslucent;
 
 			// Whitelist
 			bool bAllowed =
@@ -48,3 +62,8 @@ bool UPlacementBrushToolBase::HitTest(const FRay& Ray, FHitResult& OutHit)
 
 	return AInstancedFoliageActor::FoliageTrace(EditingWorld, OutHit, FDesiredFoliageInstance(TraceStart, TraceEnd), NAME_PlacementBrushTool, false, FilterFunc);
 }
+
+double UPlacementBrushToolBase::EstimateMaximumTargetDimension()
+{
+	return 1000.0;
+}
diff --git a/Engine/Source/Editor/AssetPlacementEdMode/Private/Tools/PlacementBrushToolBase.h b/Engine/Source/Editor/AssetPlacementEdMode/Private/Tools/PlacementBrushToolBase.h
--- a/Engine/Source/Editor/AssetPlacementEdMode/Private/Tools/PlacementBrushToolBase.h
+++ b/Engine/Source/Editor/AssetPlacementEdMode/Private/Tools/PlacementBrushToolBase.h
@@ -8,12 +8,36 @@
 
 class FRay;
 struct FHitResult;
+class UAssetPlacementSettings;
+
+UCLASS(Abstract, MinimalAPI)
+class UPlacementToolBuilderBase : public UInteractiveToolBuilder
+{
+	GENERATED_BODY()
+
+public:
+	virtual bool CanBuildTool(const FToolBuilderState& SceneState) const override;
+	virtual UInteractiveTool* BuildTool(const FToolBuilderState& SceneState) const override;
+
+	UPROPERTY()
+	UAssetPlacementSettings* PlacementSettings;
+
+protected:
+	virtual UPlacementBrushToolBase* FactoryToolInstance(UObject* Outer) const { check(false); return nullptr; }
+};
 
 UCLASS(Abstract, MinimalAPI)
 class UPlacementBrushToolBase : public UBaseBrushTool
 {
 	GENERATED_BODY()
+
+	friend class UPlacementToolBuilderBase;
 	
 public:
 	virtual bool HitTest(const FRay& Ray, FHitResult& OutHit) override;
+
+protected:
+	virtual double EstimateMaximumTargetDimension() override;
+
+	TWeakObjectPtr<UAssetPlacementSettings> PlacementSettings;
 };
diff --git a/Engine/Source/Editor/AssetPlacementEdMode/Private/Tools/PlacementClickDragToolBase.cpp b/Engine/Source/Editor/AssetPlacementEdMode/Private/Tools/PlacementClickDragToolBase.cpp
--- a/Engine/Source/Editor/AssetPlacementEdMode/Private/Tools/PlacementClickDragToolBase.cpp
+++ b/Engine/Source/Editor/AssetPlacementEdMode/Private/Tools/PlacementClickDragToolBase.cpp
@@ -15,3 +15,8 @@ void UPlacementClickDragToolBase::Setup()
 	BrushProperties->BrushSize = .5f;
 	BrushProperties->BrushFalloffAmount = 1.f;
 }
+
+double UPlacementClickDragToolBase::EstimateMaximumTargetDimension()
+{
+	return UBaseBrushTool::EstimateMaximumTargetDimension();
+}
\ No newline at end of file
diff --git a/Engine/Source/Editor/AssetPlacementEdMode/Private/Tools/PlacementClickDragToolBase.h b/Engine/Source/Editor/AssetPlacementEdMode/Private/Tools/PlacementClickDragToolBase.h
--- a/Engine/Source/Editor/AssetPlacementEdMode/Private/Tools/PlacementClickDragToolBase.h
+++ b/Engine/Source/Editor/AssetPlacementEdMode/Private/Tools/PlacementClickDragToolBase.h
@@ -12,4 +12,7 @@ class UPlacementClickDragToolBase : public UPlacementBrushToolBase
 	GENERATED_BODY()
 public:
 	virtual void Setup() override;
+
+protected:
+	virtual double EstimateMaximumTargetDimension() override;
 };
diff --git a/Engine/Source/Editor/AssetPlacementEdMode/Private/Tools/PlacementPlaceSingleTool.cpp b/Engine/Source/Editor/AssetPlacementEdMode/Private/Tools/PlacementPlaceSingleTool.cpp
--- a/Engine/Source/Editor/AssetPlacementEdMode/Private/Tools/PlacementPlaceSingleTool.cpp
+++ b/Engine/Source/Editor/AssetPlacementEdMode/Private/Tools/PlacementPlaceSingleTool.cpp
@@ -2,17 +2,41 @@
 
 #include "Tools/PlacementPlaceSingleTool.h"
 #include "UObject/Object.h"
-#include "ToolContextInterfaces.h"
-#include "InteractiveToolManager.h"
+#include "AssetPlacementSettings.h"
+#include "Subsystems/PlacementSubsystem.h"
+#include "ScopedTransaction.h"
+#include "Editor/EditorEngine.h"
+#include "Editor.h"
 
 constexpr TCHAR UPlacementModePlaceSingleTool::ToolName[];
 
-bool UPlacementModeSelectAllToolBuilder::CanBuildTool(const FToolBuilderState& SceneState) const
+UPlacementBrushToolBase* UPlacementModePlaceSingleToolBuilder::FactoryToolInstance(UObject* Outer) const
 {
-	return true;
+	return NewObject<UPlacementModePlaceSingleTool>(Outer);
 }
 
-UInteractiveTool* UPlacementModeSelectAllToolBuilder::BuildTool(const FToolBuilderState& SceneState) const
+void UPlacementModePlaceSingleTool::OnEndDrag(const FRay& Ray)
 {
-	return NewObject<UPlacementModePlaceSingleTool>(SceneState.ToolManager, UPlacementModePlaceSingleTool::ToolName);
+	Super::OnEndDrag(Ray);
+
+	if (PlacementSettings.IsValid() && PlacementSettings->PaletteItems.Num())
+	{
+		int32 ItemIndex = FMath::RandHelper(PlacementSettings->PaletteItems.Num());
+		const FPaletteItem& ItemToPlace = PlacementSettings->PaletteItems[ItemIndex];
+		if (ItemToPlace.AssetData.IsValid())
+		{
+			FAssetPlacementInfo PlacementInfo;
+			PlacementInfo.AssetToPlace = ItemToPlace.AssetData;
+			PlacementInfo.FactoryOverride = ItemToPlace.FactoryOverride;
+			PlacementInfo.FinalizedTransform = FTransform(LastBrushStamp.HitResult.ImpactPoint);
+			PlacementInfo.PreferredLevel = GEditor->GetEditorWorldContext().World()->GetCurrentLevel();
+
+			UPlacementSubsystem* PlacementSubsystem = GEditor->GetEditorSubsystem<UPlacementSubsystem>();
+			if (PlacementSubsystem)
+			{
+				FScopedTransaction Transaction(NSLOCTEXT("PlacementMode", "SinglePlaceAsset", "Place Single Asset"));
+				PlacementSubsystem->PlaceAsset(PlacementInfo, FPlacementOptions());
+			}
+		}
+	}
 }
diff --git a/Engine/Source/Editor/AssetPlacementEdMode/Private/Tools/PlacementPlaceSingleTool.h b/Engine/Source/Editor/AssetPlacementEdMode/Private/Tools/PlacementPlaceSingleTool.h
--- a/Engine/Source/Editor/AssetPlacementEdMode/Private/Tools/PlacementPlaceSingleTool.h
+++ b/Engine/Source/Editor/AssetPlacementEdMode/Private/Tools/PlacementPlaceSingleTool.h
@@ -8,13 +8,12 @@
 #include "PlacementPlaceSingleTool.generated.h"
 
 UCLASS(Transient, MinimalAPI)
-class UPlacementModeSelectAllToolBuilder : public UInteractiveToolBuilder
+class UPlacementModePlaceSingleToolBuilder : public UPlacementToolBuilderBase
 {
 	GENERATED_BODY()
 
-public:
-	virtual bool CanBuildTool(const FToolBuilderState& SceneState) const override;
-	virtual UInteractiveTool* BuildTool(const FToolBuilderState& SceneState) const override;
+protected:
+	virtual UPlacementBrushToolBase* FactoryToolInstance(UObject* Outer) const override;
 };
 
 UCLASS(MinimalAPI)
@@ -24,4 +23,6 @@ class UPlacementModePlaceSingleTool : public UPlacementClickDragToolBase
 
 public:
 	constexpr static TCHAR ToolName[] = TEXT("PlaceSingleTool");
+
+	virtual void OnEndDrag(const FRay& Ray) override;
 };
diff --git a/Engine/Source/Editor/AssetPlacementEdMode/Private/Tools/PlacementReapplySettingsTool.cpp b/Engine/Source/Editor/AssetPlacementEdMode/Private/Tools/PlacementReapplySettingsTool.cpp
deleted file mode 100644
diff --git a/Engine/Source/Editor/AssetPlacementEdMode/Private/Tools/PlacementReapplySettingsTool.h b/Engine/Source/Editor/AssetPlacementEdMode/Private/Tools/PlacementReapplySettingsTool.h
deleted file mode 100644
diff --git a/Engine/Source/Editor/AssetPlacementEdMode/Public/AssetPlacementEdMode.h b/Engine/Source/Editor/AssetPlacementEdMode/Public/AssetPlacementEdMode.h
--- a/Engine/Source/Editor/AssetPlacementEdMode/Public/AssetPlacementEdMode.h
+++ b/Engine/Source/Editor/AssetPlacementEdMode/Public/AssetPlacementEdMode.h
@@ -3,49 +3,13 @@
 #pragma once
 
 #include "CoreMinimal.h"
-#include "Tools/UEdMode.h"
+#include "Tools/LegacyEdModeWidgetHelpers.h"
 #include "Elements/Framework/TypedElementHandle.h"
 
 #include "AssetPlacementEdMode.generated.h"
 
-USTRUCT()
-struct FPaletteItem
-{
-	GENERATED_BODY()
-	FTypedHandleTypeId ElementId;
-	bool bIsEnabled;
-};
-
-UCLASS(config = EditorPerProjectUserSettings)
-class UAssetPlacementSettings : public UObject
-{
-	GENERATED_BODY()
-
-public:
-	UPROPERTY(config, EditAnywhere, Category = "Level Partition Settings")
-	bool bPlaceInCurrentLevelPartition = true;
-
-	UPROPERTY(config, EditAnywhere, Category = "Filters")
-	bool bLandscape = true;
-
-	UPROPERTY(config, EditAnywhere, Category = "Filters")
-	bool bStaticMeshes = true;
-	
-	UPROPERTY(config, EditAnywhere, Category = "Filters")
-	bool bBSP = true;
-	
-	UPROPERTY(config, EditAnywhere, Category = "Filters")
-	bool bFoliage = false;
-	
-	UPROPERTY(config, EditAnywhere, Category = "Filters")
-	bool bTranslucent = false;
-
-	UPROPERTY(config, EditAnywhere, Category = "Asset Palette")
-	TArray<FPaletteItem> PaletteItems;
-};
-
 UCLASS()
-class UAssetPlacementEdMode : public UEdMode
+class UAssetPlacementEdMode : public UBaseLegacyWidgetEdMode
 {
 	GENERATED_BODY()
 
@@ -61,16 +25,22 @@ public:
 	virtual void CreateToolkit() override;
 	virtual TMap<FName, TArray<TSharedPtr<FUICommandInfo>>> GetModeCommands() const override;
 	virtual void BindCommands() override;
-	//////////////////
+	virtual bool IsSelectionAllowed(AActor* InActor, bool bInSelection) const override;
 	// End of UEdMode interface
 	//////////////////
 
+	////////////////
+	// UBaseLegacyWidgetEdMode interface
+	virtual bool UsesPropertyWidgets() const override;
+	virtual bool ShouldDrawWidget() const override;
+	// End of UBaseLegacyWidgetEdMode interface
+	//////////////////
+
 protected:
 	enum class EPaletteFilter
 	{
 		EntirePalette,
 		ActivePaletteOnly,
-		InvalidInstances,
 	};
 
 	enum class ESelectMode
diff --git a/Engine/Source/Editor/AssetPlacementEdMode/Public/AssetPlacementEdModeCommands.h b/Engine/Source/Editor/AssetPlacementEdMode/Public/AssetPlacementEdModeCommands.h
--- a/Engine/Source/Editor/AssetPlacementEdMode/Public/AssetPlacementEdModeCommands.h
+++ b/Engine/Source/Editor/AssetPlacementEdMode/Public/AssetPlacementEdModeCommands.h
@@ -16,15 +16,12 @@ public:
 	TSharedPtr<FUICommandInfo> Select;
 	TSharedPtr<FUICommandInfo> SelectAll;
 	TSharedPtr<FUICommandInfo> Deselect;
-	TSharedPtr<FUICommandInfo> SelectInvalid;
 	TSharedPtr<FUICommandInfo> LassoSelect;
 	TSharedPtr<FUICommandInfo> Place;
-	TSharedPtr<FUICommandInfo> ReapplySettings;
 	TSharedPtr<FUICommandInfo> PlaceSingle;
 	TSharedPtr<FUICommandInfo> Fill;
 	TSharedPtr<FUICommandInfo> Erase;
 	TSharedPtr<FUICommandInfo> Delete;
-	TSharedPtr<FUICommandInfo> MoveToActivePartition;
 
 protected:
 	TMap<FName, TArray<TSharedPtr<FUICommandInfo>>> Commands;
diff --git a/Engine/Source/Editor/AssetPlacementEdMode/Public/AssetPlacementEdModeModule.h b/Engine/Source/Editor/AssetPlacementEdMode/Public/AssetPlacementEdModeModule.h
--- a/Engine/Source/Editor/AssetPlacementEdMode/Public/AssetPlacementEdModeModule.h
+++ b/Engine/Source/Editor/AssetPlacementEdMode/Public/AssetPlacementEdModeModule.h
@@ -8,6 +8,9 @@
 class FAssetPlacementEdMode : public IModuleInterface
 {
 public:
-	virtual void StartupModule();
-	virtual void ShutdownModule();
+	virtual void StartupModule() override;
+	virtual void ShutdownModule() override;
+
+protected:
+	void OnPostEngineInit();
 };
diff --git a/Engine/Source/Editor/AssetPlacementEdMode/Public/AssetPlacementEdModeToolkit.h b/Engine/Source/Editor/AssetPlacementEdMode/Public/AssetPlacementEdModeToolkit.h
--- a/Engine/Source/Editor/AssetPlacementEdMode/Public/AssetPlacementEdModeToolkit.h
+++ b/Engine/Source/Editor/AssetPlacementEdMode/Public/AssetPlacementEdModeToolkit.h
@@ -4,10 +4,9 @@
 
 #include "CoreMinimal.h"
 #include "Toolkits/BaseToolkit.h"
-#include "InteractiveTool.h"
 #include "AssetPlacementEdMode.h"
 
-class IDetailsView;
+class UAssetPlacementSettings;
 
 /**
  * This FModeToolkit just creates a basic UI panel that allows various InteractiveTools to
@@ -16,7 +15,7 @@ class IDetailsView;
 class FAssetPlacementEdModeToolkit : public FModeToolkit
 {
 public:
-	FAssetPlacementEdModeToolkit();
+	FAssetPlacementEdModeToolkit(TWeakObjectPtr<UAssetPlacementSettings> InPlacementSettings);
 	
 	// FModeToolkit interface 
 	virtual void GetToolPaletteNames(TArray<FName>& PaletteNames) const override;
@@ -24,4 +23,9 @@ public:
 	// IToolkit interface 
 	virtual FName GetToolkitFName() const override;
 	virtual FText GetBaseToolkitName() const override;
+
+protected:
+	virtual TSharedPtr<SWidget> GetInlineContent() const override;
+
+	TWeakObjectPtr<UAssetPlacementSettings> PlacementSettings;
 };
diff --git a/Engine/Source/Editor/UnrealEd/Private/Tools/UEdMode.cpp b/Engine/Source/Editor/UnrealEd/Private/Tools/UEdMode.cpp
--- a/Engine/Source/Editor/UnrealEd/Private/Tools/UEdMode.cpp
+++ b/Engine/Source/Editor/UnrealEd/Private/Tools/UEdMode.cpp
@@ -224,6 +224,13 @@ void UEdMode::Enter()
 	GetToolManager()->OnToolStarted.AddUObject(this, &UEdMode::OnToolStarted);
 	GetToolManager()->OnToolEnded.AddUObject(this, &UEdMode::OnToolEnded);
 
+	// Create the settings object so that the toolkit has access to the object we are going to use at creation time
+	if (SettingsClass.IsValid())
+	{
+		UClass* LoadedSettingsObject = SettingsClass.LoadSynchronous();
+		SettingsObject = NewObject<UObject>(this, LoadedSettingsObject);
+	}
+
 	// Now that the context is ready, make the toolkit
 	CreateToolkit();
 	if (Toolkit.IsValid())
@@ -233,12 +240,6 @@ void UEdMode::Enter()
 
 	BindCommands();
 
-	if (SettingsClass.IsValid())
-	{
-		UClass* LoadedSettingsObject = SettingsClass.LoadSynchronous();
-		SettingsObject = NewObject<UObject>(this, LoadedSettingsObject);
-	}
-
 	if (SettingsObject)
 	{
 		SettingsObject->LoadConfig();
