Asset Placement Mode's paint tool implemented.

diff --git a/Engine/Source/Editor/AssetPlacementEdMode/Private/AssetPlacementEdMode.cpp b/Engine/Source/Editor/AssetPlacementEdMode/Private/AssetPlacementEdMode.cpp
--- a/Engine/Source/Editor/AssetPlacementEdMode/Private/AssetPlacementEdMode.cpp
+++ b/Engine/Source/Editor/AssetPlacementEdMode/Private/AssetPlacementEdMode.cpp
@@ -48,7 +48,10 @@ void UAssetPlacementEdMode::Enter()
 	const FAssetPlacementEdModeCommands& PlacementModeCommands = FAssetPlacementEdModeCommands::Get();
 	RegisterTool(PlacementModeCommands.Select, UPlacementModeSelectTool::ToolName, NewObject<UPlacementModeSelectToolBuilder>(this));
 	RegisterTool(PlacementModeCommands.LassoSelect, UPlacementModeLassoSelectTool::ToolName, NewObject<UPlacementModeLassoSelectToolBuilder>(this));
-	RegisterTool(PlacementModeCommands.Place, UPlacementModePlacementTool::ToolName, NewObject<UPlacementModePlacementToolBuilder>(this));
+
+	UPlacementModePlacementToolBuilder* PlaceTool = NewObject<UPlacementModePlacementToolBuilder>(this);
+	PlaceTool->PlacementSettings = Cast<UAssetPlacementSettings>(SettingsObject);
+	RegisterTool(PlacementModeCommands.Place, UPlacementModePlacementTool::ToolName, PlaceTool);
 
 	UPlacementModePlaceSingleToolBuilder* SinglePlaceTool = NewObject<UPlacementModePlaceSingleToolBuilder>(this);
 	SinglePlaceTool->PlacementSettings = Cast<UAssetPlacementSettings>(SettingsObject);
diff --git a/Engine/Source/Editor/AssetPlacementEdMode/Private/AssetPlacementSettings.h b/Engine/Source/Editor/AssetPlacementEdMode/Private/AssetPlacementSettings.h
--- a/Engine/Source/Editor/AssetPlacementEdMode/Private/AssetPlacementSettings.h
+++ b/Engine/Source/Editor/AssetPlacementEdMode/Private/AssetPlacementSettings.h
@@ -45,6 +45,15 @@ public:
 	UPROPERTY(config, EditAnywhere, Category = "Filters")
 	bool bTranslucent = false;
 
+	UPROPERTY(config, EditAnywhere, Category = "Placement")
+	bool bAllowRandomRotation = true;
+
+	UPROPERTY(config, EditAnywhere, Category = "Placement")
+	bool bAllowAlignToNormal = true;
+
+	UPROPERTY(config, EditAnywhere, Category = "Placement")
+	bool bAllowRandomScale = true;
+
 	// todo: asset data does not serialize out correctly
 	// maybe save soft object pointers, and convert in the UI to asset data?
 	//UPROPERTY(config)
diff --git a/Engine/Source/Editor/AssetPlacementEdMode/Private/Tools/PlacementBrushToolBase.cpp b/Engine/Source/Editor/AssetPlacementEdMode/Private/Tools/PlacementBrushToolBase.cpp
--- a/Engine/Source/Editor/AssetPlacementEdMode/Private/Tools/PlacementBrushToolBase.cpp
+++ b/Engine/Source/Editor/AssetPlacementEdMode/Private/Tools/PlacementBrushToolBase.cpp
@@ -29,9 +29,20 @@ UInteractiveTool* UPlacementToolBuilderBase::BuildTool(const FToolBuilderState&
 
 bool UPlacementBrushToolBase::HitTest(const FRay& Ray, FHitResult& OutHit)
 {
-	UWorld* EditingWorld = GetToolManager()->GetWorld();
 	const FVector TraceStart(Ray.Origin);
 	const FVector TraceEnd(Ray.Origin + Ray.Direction * HALF_WORLD_MAX);
+
+	return FindHitResultWithStartAndEndTraceVectors(OutHit, TraceStart, TraceEnd);
+}
+
+double UPlacementBrushToolBase::EstimateMaximumTargetDimension()
+{
+	return 1000.0;
+}
+
+bool UPlacementBrushToolBase::FindHitResultWithStartAndEndTraceVectors(FHitResult& OutHit, const FVector& TraceStart, const FVector& TraceEnd, float TraceRadius)
+{
+	UWorld* EditingWorld = GetToolManager()->GetWorld();
 	constexpr TCHAR NAME_PlacementBrushTool[] = TEXT("PlacementBrushTool");
 
 	auto FilterFunc = [this](const UPrimitiveComponent* InComponent) {
@@ -60,10 +71,66 @@ bool UPlacementBrushToolBase::HitTest(const FRay& Ray, FHitResult& OutHit)
 
 		return false; };
 
-	return AInstancedFoliageActor::FoliageTrace(EditingWorld, OutHit, FDesiredFoliageInstance(TraceStart, TraceEnd), NAME_PlacementBrushTool, false, FilterFunc);
+	return AInstancedFoliageActor::FoliageTrace(EditingWorld, OutHit, FDesiredFoliageInstance(TraceStart, TraceEnd, TraceRadius), NAME_PlacementBrushTool, false, FilterFunc);
 }
 
-double UPlacementBrushToolBase::EstimateMaximumTargetDimension()
+FTransform UPlacementBrushToolBase::GetFinalTransformFromHitLocationAndNormal(const FVector& InLocation, const FVector& InNormal)
 {
-	return 1000.0;
+	FTransform FinalizedTransform(InLocation);
+
+	if (!PlacementSettings.IsValid())
+	{
+		return FinalizedTransform;
+	}
+
+	// For now, just apply a random yaw to any placed object, until we have per object settings for random pitch and yaw angles.
+	if (PlacementSettings->bAllowRandomRotation)
+	{
+		FRotator UpdatedRotation = FinalizedTransform.Rotator();
+
+		// UpdatedRotation = FRotator(FMath::FRand() * ItemToPlace->RandomPitchAngle, 0.f, 0.f);
+
+		// if (ItemToPlace.bUseRandomYaw)
+		UpdatedRotation.Yaw = FMath::FRand() * 360.f;
+
+		FinalizedTransform.SetRotation(UpdatedRotation.Quaternion());
+	}
+
+	// Align to normal
+	if (PlacementSettings->bAllowAlignToNormal)
+	{
+		FRotator AlignRotation = InNormal.Rotation();
+		// Static meshes are authored along the vertical axis rather than the X axis, so we add 90 degrees to the static mesh's Pitch.
+		AlignRotation.Pitch -= 90.f;
+		// Clamp its value inside +/- one rotation
+		AlignRotation.Pitch = FRotator::NormalizeAxis(AlignRotation.Pitch);
+
+		// limit the maximum pitch angle if it's > 0.
+		// For now, just set the align max angle to a constant value, until it can be pulled from per object settings
+		constexpr float AlignMaxAngle = 0.0f;
+		if (AlignMaxAngle > 0.f)
+		{
+			int32 MaxPitch = AlignMaxAngle;
+			if (AlignRotation.Pitch > MaxPitch)
+			{
+				AlignRotation.Pitch = MaxPitch;
+			}
+			else if (AlignRotation.Pitch < -MaxPitch)
+			{
+				AlignRotation.Pitch = -MaxPitch;
+			}
+		}
+
+		FinalizedTransform.SetRotation(FQuat(AlignRotation) * FinalizedTransform.GetRotation());
+	}
+
+	if (PlacementSettings->bAllowRandomScale)
+	{
+		// Until we have per object settings, just use a uniform scale, clamped from half to double size
+		FFloatInterval ScaleRange(0.5f, 2.0f);
+		FVector NewScale(ScaleRange.Interpolate(FMath::FRand()));
+		FinalizedTransform.SetScale3D(NewScale);
+	}
+
+	return FinalizedTransform;
 }
diff --git a/Engine/Source/Editor/AssetPlacementEdMode/Private/Tools/PlacementBrushToolBase.h b/Engine/Source/Editor/AssetPlacementEdMode/Private/Tools/PlacementBrushToolBase.h
--- a/Engine/Source/Editor/AssetPlacementEdMode/Private/Tools/PlacementBrushToolBase.h
+++ b/Engine/Source/Editor/AssetPlacementEdMode/Private/Tools/PlacementBrushToolBase.h
@@ -38,6 +38,8 @@ public:
 
 protected:
 	virtual double EstimateMaximumTargetDimension() override;
+	bool FindHitResultWithStartAndEndTraceVectors(FHitResult& OutHit, const FVector& TraceStart, const FVector& TraceEnd, float TraceRadius = 0.0f);
+	FTransform GetFinalTransformFromHitLocationAndNormal(const FVector& InLocation, const FVector& InNormal);
 
 	TWeakObjectPtr<UAssetPlacementSettings> PlacementSettings;
 };
diff --git a/Engine/Source/Editor/AssetPlacementEdMode/Private/Tools/PlacementPlaceSingleTool.cpp b/Engine/Source/Editor/AssetPlacementEdMode/Private/Tools/PlacementPlaceSingleTool.cpp
--- a/Engine/Source/Editor/AssetPlacementEdMode/Private/Tools/PlacementPlaceSingleTool.cpp
+++ b/Engine/Source/Editor/AssetPlacementEdMode/Private/Tools/PlacementPlaceSingleTool.cpp
@@ -28,7 +28,7 @@ void UPlacementModePlaceSingleTool::OnEndDrag(const FRay& Ray)
 			FAssetPlacementInfo PlacementInfo;
 			PlacementInfo.AssetToPlace = ItemToPlace.AssetData;
 			PlacementInfo.FactoryOverride = ItemToPlace.FactoryOverride;
-			PlacementInfo.FinalizedTransform = FTransform(LastBrushStamp.HitResult.ImpactPoint);
+			PlacementInfo.FinalizedTransform = GetFinalTransformFromHitLocationAndNormal(LastBrushStamp.HitResult.ImpactPoint, LastBrushStamp.HitResult.ImpactNormal);
 			PlacementInfo.PreferredLevel = GEditor->GetEditorWorldContext().World()->GetCurrentLevel();
 
 			FPlacementOptions PlacementOptions;
diff --git a/Engine/Source/Editor/AssetPlacementEdMode/Private/Tools/PlacementPlaceTool.cpp b/Engine/Source/Editor/AssetPlacementEdMode/Private/Tools/PlacementPlaceTool.cpp
--- a/Engine/Source/Editor/AssetPlacementEdMode/Private/Tools/PlacementPlaceTool.cpp
+++ b/Engine/Source/Editor/AssetPlacementEdMode/Private/Tools/PlacementPlaceTool.cpp
@@ -2,17 +2,131 @@
 
 #include "Tools/PlacementPlaceTool.h"
 #include "UObject/Object.h"
-#include "ToolContextInterfaces.h"
-#include "InteractiveToolManager.h"
+#include "AssetPlacementSettings.h"
+#include "Subsystems/PlacementSubsystem.h"
+#include "Math/UnrealMathUtility.h"
+#include "BaseGizmos/BrushStampIndicator.h"
+#include "Editor.h"
+#include "InstancedFoliage.h"
+#include "ScopedTransaction.h"
 
 constexpr TCHAR UPlacementModePlacementTool::ToolName[];
 
-bool UPlacementModePlacementToolBuilder::CanBuildTool(const FToolBuilderState& SceneState) const
+UPlacementBrushToolBase* UPlacementModePlacementToolBuilder::FactoryToolInstance(UObject* Outer) const
 {
-	return true;
+	return NewObject<UPlacementModePlacementTool>(Outer);
 }
 
-UInteractiveTool* UPlacementModePlacementToolBuilder::BuildTool(const FToolBuilderState& SceneState) const
+void UPlacementModePlacementTool::OnBeginDrag(const FRay& Ray)
 {
-	return NewObject<UPlacementModePlacementTool>(SceneState.ToolManager, UPlacementModePlacementTool::ToolName);
+	Super::OnBeginDrag(Ray);
+
+	TransactionScope = MakeUnique<FScopedTransaction>(NSLOCTEXT("PlacementMode", "PaintAssets", "Paint Asset Stroke"));
+}
+
+void UPlacementModePlacementTool::OnEndDrag(const FRay& Ray)
+{
+	TransactionScope.Reset();
+
+	Super::OnEndDrag(Ray);
+}
+
+void UPlacementModePlacementTool::OnTick(float DeltaTime)
+{
+	if (!bInBrushStroke)
+	{
+		return;
+	}
+
+	if (BrushProperties && PlacementSettings.IsValid() && PlacementSettings->PaletteItems.Num())
+	{
+		FVector BrushLocation = LastBrushStamp.WorldPosition;
+
+		// Assume a default density of 100 *  whatever the user has selected as brush size.
+		float DefaultDensity = 100.0f * BrushProperties->BrushSize;
+		// This is the total set of instances disregarding parameters like slope, height or layer.
+		float DesiredInstanceCountFloat = DefaultDensity /** todo: individual item's Density*/ * BrushProperties->BrushStrength;
+		// Allow a single instance with a random chance, if the brush is smaller than the density
+		int32 DesiredInstanceCount = DesiredInstanceCountFloat > 1.f ? FMath::RoundToInt(DesiredInstanceCountFloat) : FMath::FRand() < DesiredInstanceCountFloat ? 1 : 0;
+
+		// Todo: save the instance hash per tile for the editor, so that paints don't continually add instances and surpass desired density (partition actor?)
+		FFoliageInstanceHash PotentialInstanceHash(7);
+		TArray<FVector> PotentialInstanceLocations;
+
+		TArray<FAssetPlacementInfo> DesiredInfoSpawnLocations;
+		DesiredInfoSpawnLocations.Reserve(DesiredInstanceCount);
+		PotentialInstanceLocations.Empty(DesiredInstanceCount);
+		for (int32 Index = 0; Index < DesiredInstanceCount; ++Index)
+		{
+			FVector Start, End;
+			GetRandomVectorInBrush(Start, End);
+			FHitResult AdjustedHitResult;
+			FindHitResultWithStartAndEndTraceVectors(AdjustedHitResult, Start, End);
+			FVector SpawnLocation = AdjustedHitResult.ImpactPoint;
+			FVector SpawnNormal = AdjustedHitResult.ImpactNormal;
+
+			bool bValidInstance = true;
+			auto TempInstances = PotentialInstanceHash.GetInstancesOverlappingBox(FBox::BuildAABB(BrushLocation, FVector(BrushProperties->BrushRadius)));
+			for (int32 InstanceIndex : TempInstances)
+			{
+				if ((PotentialInstanceLocations[InstanceIndex] - BrushLocation).SizeSquared() < FMath::Square(BrushProperties->BrushRadius))
+				{
+					bValidInstance = false;
+					break;
+				}
+			}
+
+			if (!bValidInstance)
+			{
+				continue;
+			}
+
+			int32 PotentialIdx = PotentialInstanceLocations.Add(SpawnLocation);
+			PotentialInstanceHash.InsertInstance(SpawnLocation, PotentialIdx);
+
+			int32 ItemIndex = FMath::RandHelper(PlacementSettings->PaletteItems.Num());
+			const FPaletteItem& ItemToPlace = PlacementSettings->PaletteItems[ItemIndex];
+			FAssetPlacementInfo NewInfo;
+			NewInfo.AssetToPlace = ItemToPlace.AssetData;
+			NewInfo.PreferredLevel = GEditor->GetEditorWorldContext().World()->GetCurrentLevel();
+			NewInfo.FinalizedTransform = GetFinalTransformFromHitLocationAndNormal(SpawnLocation, SpawnNormal);
+
+			DesiredInfoSpawnLocations.Emplace(NewInfo);
+		}
+
+		UPlacementSubsystem* PlacementSubsystem = GEditor->GetEditorSubsystem<UPlacementSubsystem>();
+		if (PlacementSubsystem)
+		{
+			FPlacementOptions PlacementOptions;
+			PlacementOptions.bPreferBatchPlacement = true;
+			PlacementOptions.bPreferInstancedPlacement = true;
+			PlacementSubsystem->PlaceAssets(DesiredInfoSpawnLocations, PlacementOptions);
+		}
+	}
+}
+
+void UPlacementModePlacementTool::GetRandomVectorInBrush(FVector& OutStart, FVector& OutEnd)
+{
+	if (!BrushProperties)
+	{
+		return;
+	}
+	FVector BrushNormal = LastBrushStamp.WorldNormal;
+	FVector BrushLocation = LastBrushStamp.WorldPosition;
+	float BrushRadius = BrushProperties->BrushRadius;
+
+	// Find Rx and Ry inside the unit circle
+	float Ru = (2.f * FMath::FRand() - 1.f);
+	float Rv = (2.f * FMath::FRand() - 1.f) * FMath::Sqrt(1.f - FMath::Square(Ru));
+
+	// find random point in circle through brush location on the same plane to brush location hit surface normal
+	FVector U, V;
+	BrushNormal.FindBestAxisVectors(U, V);
+	FVector Point = Ru * U + Rv * V;
+
+	// find distance to surface of sphere brush from this point
+	FVector Rw = FMath::Sqrt(FMath::Max(1.f - (FMath::Square(Ru) + FMath::Square(Rv)), 0.001f)) * BrushNormal;
+
+	OutStart = BrushLocation + BrushRadius * (Point + Rw);
+	OutEnd = BrushLocation + BrushRadius * (Point - Rw);
 }
diff --git a/Engine/Source/Editor/AssetPlacementEdMode/Private/Tools/PlacementPlaceTool.h b/Engine/Source/Editor/AssetPlacementEdMode/Private/Tools/PlacementPlaceTool.h
--- a/Engine/Source/Editor/AssetPlacementEdMode/Private/Tools/PlacementPlaceTool.h
+++ b/Engine/Source/Editor/AssetPlacementEdMode/Private/Tools/PlacementPlaceTool.h
@@ -7,14 +7,15 @@
 
 #include "PlacementPlaceTool.generated.h"
 
+class FScopedTransaction;
+
 UCLASS(Transient, MinimalAPI)
-class UPlacementModePlacementToolBuilder : public UInteractiveToolBuilder
+class UPlacementModePlacementToolBuilder : public UPlacementToolBuilderBase
 {
 	GENERATED_BODY()
 
 public:
-	virtual bool CanBuildTool(const FToolBuilderState& SceneState) const override;
-	virtual UInteractiveTool* BuildTool(const FToolBuilderState& SceneState) const override;
+	virtual UPlacementBrushToolBase* FactoryToolInstance(UObject* Outer) const override;
 };
 
 UCLASS(MinimalAPI)
@@ -23,4 +24,14 @@ class UPlacementModePlacementTool : public UPlacementBrushToolBase
 	GENERATED_BODY()
 public:
 	constexpr static TCHAR ToolName[] = TEXT("PlaceTool");
+
+
+protected:
+	virtual void OnBeginDrag(const FRay& Ray) override;
+	virtual void OnEndDrag(const FRay& Ray) override;
+	virtual void OnTick(float DeltaTime) override;
+
+	void GetRandomVectorInBrush(FVector& OutStart, FVector& OutEnd);
+
+	TUniquePtr<FScopedTransaction> TransactionScope;
 };
