Implemented Erase Tool in asset placement mode.

diff --git a/Engine/Source/Editor/AssetPlacementEdMode/Private/AssetPlacementEdMode.cpp b/Engine/Source/Editor/AssetPlacementEdMode/Private/AssetPlacementEdMode.cpp
--- a/Engine/Source/Editor/AssetPlacementEdMode/Private/AssetPlacementEdMode.cpp
+++ b/Engine/Source/Editor/AssetPlacementEdMode/Private/AssetPlacementEdMode.cpp
@@ -5,7 +5,6 @@
 #include "InteractiveToolManager.h"
 #include "AssetPlacementEdModeStyle.h"
 #include "AssetPlacementEdModeCommands.h"
-#include "AssetPlacementSettings.h"
 #include "EdModeInteractiveToolsContext.h"
 #include "EditorModeManager.h"
 #include "Selection.h"
@@ -45,26 +44,32 @@ void UAssetPlacementEdMode::Enter()
 {
 	Super::Enter();
 
+	UAssetPlacementSettings* PlacementSettings = Cast<UAssetPlacementSettings>(SettingsObject);
+
 	const FAssetPlacementEdModeCommands& PlacementModeCommands = FAssetPlacementEdModeCommands::Get();
 	RegisterTool(PlacementModeCommands.Select, UPlacementModeSelectTool::ToolName, NewObject<UPlacementModeSelectToolBuilder>(this));
 	RegisterTool(PlacementModeCommands.LassoSelect, UPlacementModeLassoSelectTool::ToolName, NewObject<UPlacementModeLassoSelectToolBuilder>(this));
 
-	UPlacementModePlacementToolBuilder* PlaceTool = NewObject<UPlacementModePlacementToolBuilder>(this);
-	PlaceTool->PlacementSettings = Cast<UAssetPlacementSettings>(SettingsObject);
-	RegisterTool(PlacementModeCommands.Place, UPlacementModePlacementTool::ToolName, PlaceTool);
+	UPlacementModePlacementToolBuilder* PlaceToolBuilder = NewObject<UPlacementModePlacementToolBuilder>(this);
+	PlaceToolBuilder->PlacementSettings = Cast<UAssetPlacementSettings>(SettingsObject);
+	RegisterTool(PlacementModeCommands.Place, UPlacementModePlacementTool::ToolName, PlaceToolBuilder);
 
-	UPlacementModePlaceSingleToolBuilder* SinglePlaceTool = NewObject<UPlacementModePlaceSingleToolBuilder>(this);
-	SinglePlaceTool->PlacementSettings = Cast<UAssetPlacementSettings>(SettingsObject);
-	RegisterTool(PlacementModeCommands.PlaceSingle, UPlacementModePlaceSingleTool::ToolName, SinglePlaceTool);
+	UPlacementModePlaceSingleToolBuilder* SinglePlaceToolBuilder = NewObject<UPlacementModePlaceSingleToolBuilder>(this);
+	SinglePlaceToolBuilder->PlacementSettings = Cast<UAssetPlacementSettings>(SettingsObject);
+	RegisterTool(PlacementModeCommands.PlaceSingle, UPlacementModePlaceSingleTool::ToolName, SinglePlaceToolBuilder);
 
 	RegisterTool(PlacementModeCommands.Fill, UPlacementModeFillTool::ToolName, NewObject<UPlacementModeFillToolBuilder>(this));
-	RegisterTool(PlacementModeCommands.Erase, UPlacementModeEraseTool::ToolName, NewObject<UPlacementModeEraseToolBuilder>(this));
+
+	UPlacementModeEraseToolBuilder* EraseToolBuilder = NewObject<UPlacementModeEraseToolBuilder>(this);
+	EraseToolBuilder->PlacementSettings = Cast<UAssetPlacementSettings>(SettingsObject);
+	RegisterTool(PlacementModeCommands.Erase, UPlacementModeEraseTool::ToolName, EraseToolBuilder);
 
 	ToolsContext->ToolManager->SelectActiveToolType(EToolSide::Mouse, UPlacementModeSelectTool::ToolName);
 	ToolsContext->ToolManager->ActivateTool(EToolSide::Mouse);
 
-	// TODO - Update selection set based on valid things in the palette
-	Owner->GetSelectedObjects()->GetElementSelectionSet()->ClearSelection(FTypedElementSelectionOptions());
+	// TODO - Stash existing selection and update to filter on valid things in the palette
+	Owner->GetSelectedActors()->GetElementSelectionSet()->ClearSelection(FTypedElementSelectionOptions());
+	Owner->GetSelectedComponents()->GetElementSelectionSet()->ClearSelection(FTypedElementSelectionOptions());
 }
 
 void UAssetPlacementEdMode::CreateToolkit()
@@ -103,23 +108,10 @@ bool UAssetPlacementEdMode::IsSelectionAllowed(AActor* InActor, bool bInSelectio
 		return false;
 	}
 
-	FTypedElementHandle ActorHandle = UEngineElementsLibrary::AcquireEditorActorElementHandle(InActor);
-	if (!ActorHandle.IsSet())
-	{
-		return false;
-	}
-
 	UAssetPlacementSettings* PlacementSettings = Cast<UAssetPlacementSettings>(SettingsObject);
-	for (FPaletteItem& Item : PlacementSettings->PaletteItems)
-	{
-		FAssetData FoundAssetDataFromFactory = Item.FactoryOverride->GetAssetDataFromElementHandle(ActorHandle);
-		if (FoundAssetDataFromFactory == Item.AssetData)
-		{
-			return true;
-		}
-	}
+	FTypedElementHandle ActorHandle = UEngineElementsLibrary::AcquireEditorActorElementHandle(InActor);
 
-	return false;
+	return DoesPaletteSupportElement(ActorHandle, PlacementSettings->PaletteItems);
 }
 
 bool UAssetPlacementEdMode::UsesPropertyWidgets() const
@@ -142,6 +134,25 @@ bool UAssetPlacementEdMode::ShouldDrawWidget() const
 	return Super::ShouldDrawWidget();
 }
 
+bool UAssetPlacementEdMode::DoesPaletteSupportElement(const FTypedElementHandle& InElementToCheck, const TArray<FPaletteItem>& InPaletteToCheck)
+{
+	if (!InElementToCheck.IsSet())
+	{
+		return false;
+	}
+
+	for (const FPaletteItem& Item : InPaletteToCheck)
+	{
+		FAssetData FoundAssetDataFromFactory = Item.FactoryOverride->GetAssetDataFromElementHandle(InElementToCheck);
+		if (FoundAssetDataFromFactory == Item.AssetData)
+		{
+			return true;
+		}
+	}
+
+	return false;
+}
+
 void UAssetPlacementEdMode::SelectAssets(EPaletteFilter InSelectAllType, ESelectMode InSelectMode)
 {
 }
diff --git a/Engine/Source/Editor/AssetPlacementEdMode/Private/Tools/PlacementBrushToolBase.cpp b/Engine/Source/Editor/AssetPlacementEdMode/Private/Tools/PlacementBrushToolBase.cpp
--- a/Engine/Source/Editor/AssetPlacementEdMode/Private/Tools/PlacementBrushToolBase.cpp
+++ b/Engine/Source/Editor/AssetPlacementEdMode/Private/Tools/PlacementBrushToolBase.cpp
@@ -13,6 +13,10 @@
 #include "LandscapeHeightfieldCollisionComponent.h"
 #include "FoliageInstancedStaticMeshComponent.h"
 #include "AssetPlacementSettings.h"
+#include "Elements/Framework/EngineElementsLibrary.h"
+#include "AssetPlacementEdMode.h"
+#include "ActorPartition/ActorPartitionSubsystem.h"
+#include "Editor.h"
 
 bool UPlacementToolBuilderBase::CanBuildTool(const FToolBuilderState& SceneState) const
 {
@@ -35,6 +39,11 @@ bool UPlacementBrushToolBase::HitTest(const FRay& Ray, FHitResult& OutHit)
 	return FindHitResultWithStartAndEndTraceVectors(OutHit, TraceStart, TraceEnd);
 }
 
+bool UPlacementBrushToolBase::AreAllTargetsValid() const
+{
+	return Target ? Target->IsValid() : true;
+}
+
 double UPlacementBrushToolBase::EstimateMaximumTargetDimension()
 {
 	return 1000.0;
@@ -134,3 +143,76 @@ FTransform UPlacementBrushToolBase::GetFinalTransformFromHitLocationAndNormal(co
 
 	return FinalizedTransform;
 }
+
+TArray<FTypedElementHandle> UPlacementBrushToolBase::GetElementsInBrushRadius() const
+{
+	TArray<FTypedElementHandle> ElementHandles;
+	FCollisionQueryParams QueryParams(TEXT("PlacementBrushTool"), SCENE_QUERY_STAT_ONLY(IFA_FoliageTrace), true);
+	QueryParams.bReturnFaceIndex = false;
+	TArray<FHitResult> Hits;
+	FCollisionShape BrushSphere;
+	BrushSphere.SetSphere(LastBrushStamp.Radius);
+
+	const FVector TraceStart(LastWorldRay.Origin);
+	const FVector TraceEnd(LastWorldRay.Origin + LastWorldRay.Direction * HALF_WORLD_MAX);
+
+	GetToolManager()->GetWorld()->SweepMultiByObjectType(Hits, TraceStart, TraceEnd, FQuat::Identity,
+		FCollisionObjectQueryParams(FCollisionObjectQueryParams::InitType::AllObjects), BrushSphere, QueryParams);
+
+	for (const FHitResult& Hit : Hits)
+	{
+		const UPrimitiveComponent* HitComponent = Hit.GetComponent();
+		check(HitComponent);
+
+		// In the editor traces can hit "No Collision" type actors, so ugh. (ignore these)
+		if (!HitComponent->IsQueryCollisionEnabled() || HitComponent->GetCollisionResponseToChannel(ECC_WorldStatic) != ECR_Block)
+		{
+			continue;
+		}
+
+		// Don't place foliage on invisible walls / triggers / volumes
+		if (HitComponent->IsA<UBrushComponent>())
+		{
+			continue;
+		}
+
+		const FActorInstanceHandle& HitObjectHandle = Hit.HitObjectHandle;
+		const AActor* HitActor = Hit.HitObjectHandle.FetchActor();
+		if (HitActor)
+		{
+			FTypedElementHandle ActorHandle = UEngineElementsLibrary::AcquireEditorActorElementHandle(HitActor);
+			if (UAssetPlacementEdMode::DoesPaletteSupportElement(ActorHandle, PlacementSettings->PaletteItems))
+			{
+				ElementHandles.Emplace(ActorHandle);
+			}
+		}
+	}
+
+	// Handle the IFA for the brush stroke level
+	if (UActorPartitionSubsystem* PartitionSubsystem = UWorld::GetSubsystem<UActorPartitionSubsystem>(GEditor->GetEditorWorldContext().World()))
+	{
+		constexpr bool bCreatePartitionActorIfMissing = false;
+		FActorPartitionGetParams PartitionActorFindParams(AInstancedFoliageActor::StaticClass(), bCreatePartitionActorIfMissing, GEditor->GetEditorWorldContext().World()->GetCurrentLevel(), LastBrushStamp.WorldPosition);
+		if (AInstancedFoliageActor* FoliageActor = Cast<AInstancedFoliageActor>(PartitionSubsystem->GetActor(PartitionActorFindParams)))
+		{
+			for (auto& FoliageInfo : FoliageActor->FoliageInfos)
+			{
+				FTypedElementHandle SourceObjectHandle = UEngineElementsLibrary::AcquireEditorObjectElementHandle(FoliageInfo.Key->GetSource());
+				if (UAssetPlacementEdMode::DoesPaletteSupportElement(SourceObjectHandle, PlacementSettings->PaletteItems))
+				{
+					TArray<int32> Instances;
+					FSphere SphereToCheck(LastBrushStamp.WorldPosition, LastBrushStamp.Radius);
+					FoliageInfo.Value->GetInstancesInsideSphere(SphereToCheck, Instances);
+					if (Instances.Num())
+					{
+						// For now, return the whole foliage actor, and allow the calling code to drill down, since we do not have element handles at the instance level just yet
+						ElementHandles.Emplace(UEngineElementsLibrary::AcquireEditorActorElementHandle(FoliageActor));
+						break;
+					}
+				}
+			}
+		}
+	}
+
+	return ElementHandles;
+}
diff --git a/Engine/Source/Editor/AssetPlacementEdMode/Private/Tools/PlacementBrushToolBase.h b/Engine/Source/Editor/AssetPlacementEdMode/Private/Tools/PlacementBrushToolBase.h
--- a/Engine/Source/Editor/AssetPlacementEdMode/Private/Tools/PlacementBrushToolBase.h
+++ b/Engine/Source/Editor/AssetPlacementEdMode/Private/Tools/PlacementBrushToolBase.h
@@ -3,6 +3,7 @@
 #pragma once
 
 #include "BaseTools/BaseBrushTool.h"
+#include "Elements/Framework/TypedElementHandle.h"
 
 #include "PlacementBrushToolBase.generated.h"
 
@@ -23,7 +24,7 @@ public:
 	UAssetPlacementSettings* PlacementSettings;
 
 protected:
-	virtual UPlacementBrushToolBase* FactoryToolInstance(UObject* Outer) const { check(false); return nullptr; }
+	virtual UPlacementBrushToolBase* FactoryToolInstance(UObject* Outer) const PURE_VIRTUAL(UPlacementToolBuilderBase::FactoryToolInstance, return nullptr; );
 };
 
 UCLASS(Abstract, MinimalAPI)
@@ -35,11 +36,13 @@ class UPlacementBrushToolBase : public UBaseBrushTool
 	
 public:
 	virtual bool HitTest(const FRay& Ray, FHitResult& OutHit) override;
+	virtual bool AreAllTargetsValid() const override;
 
 protected:
 	virtual double EstimateMaximumTargetDimension() override;
 	bool FindHitResultWithStartAndEndTraceVectors(FHitResult& OutHit, const FVector& TraceStart, const FVector& TraceEnd, float TraceRadius = 0.0f);
 	FTransform GetFinalTransformFromHitLocationAndNormal(const FVector& InLocation, const FVector& InNormal);
+	TArray<FTypedElementHandle> GetElementsInBrushRadius() const;
 
 	TWeakObjectPtr<UAssetPlacementSettings> PlacementSettings;
 };
diff --git a/Engine/Source/Editor/AssetPlacementEdMode/Private/Tools/PlacementEraseTool.cpp b/Engine/Source/Editor/AssetPlacementEdMode/Private/Tools/PlacementEraseTool.cpp
--- a/Engine/Source/Editor/AssetPlacementEdMode/Private/Tools/PlacementEraseTool.cpp
+++ b/Engine/Source/Editor/AssetPlacementEdMode/Private/Tools/PlacementEraseTool.cpp
@@ -4,15 +4,69 @@
 #include "UObject/Object.h"
 #include "ToolContextInterfaces.h"
 #include "InteractiveToolManager.h"
+#include "Subsystems/EditorActorSubsystem.h"
+#include "Elements/Framework/EngineElementsLibrary.h"
+#include "Elements/Framework/TypedElementHandle.h"
+#include "Elements/Actor/ActorElementData.h"
+#include "Editor.h"
+#include "InstancedFoliageActor.h"
+#include "AssetPlacementEdMode.h"
+#include "AssetPlacementSettings.h"
 
 constexpr TCHAR UPlacementModeEraseTool::ToolName[];
 
-bool UPlacementModeEraseToolBuilder::CanBuildTool(const FToolBuilderState& SceneState) const
+UPlacementBrushToolBase* UPlacementModeEraseToolBuilder::FactoryToolInstance(UObject* Outer) const
 {
-	return true;
+	return NewObject<UPlacementModeEraseTool>(Outer);
 }
 
-UInteractiveTool* UPlacementModeEraseToolBuilder::BuildTool(const FToolBuilderState& SceneState) const
+void UPlacementModeEraseTool::OnBeginDrag(const FRay& Ray)
 {
-	return NewObject<UPlacementModeEraseTool>(SceneState.ToolManager, UPlacementModeEraseTool::ToolName);
+	Super::OnBeginDrag(Ray);
+
+	GetToolManager()->BeginUndoTransaction(NSLOCTEXT("AssetPlacementEdMode", "BrushErase", "Erase Painted Elements"));
+}
+
+void UPlacementModeEraseTool::OnEndDrag(const FRay& Ray)
+{
+	GetToolManager()->EndUndoTransaction();
+
+	Super::OnEndDrag(Ray);
+}
+
+void UPlacementModeEraseTool::OnTick(float DeltaTime)
+{
+	if (!bInBrushStroke)
+	{
+		return;
+	}
+
+	UEditorActorSubsystem* ActorSubsystem = GEditor->GetEditorSubsystem<UEditorActorSubsystem>();
+	TArray<FTypedElementHandle> HitElements = GetElementsInBrushRadius();
+	for (const FTypedElementHandle& HitElement : HitElements)
+	{
+		const FActorElementData* ActorData = HitElement.GetData<FActorElementData>();
+		if (ActorData && ActorData->Actor)
+		{
+			// Since the foliage static mesh instances do not currently operate with element handles, we have to drill in manually here.
+			if (AInstancedFoliageActor* FoliageActor = Cast<AInstancedFoliageActor>(ActorData->Actor))
+			{
+				for (auto& FoliageInfo : FoliageActor->FoliageInfos)
+				{
+					FTypedElementHandle SourceObjectHandle = UEngineElementsLibrary::AcquireEditorObjectElementHandle(FoliageInfo.Key->GetSource());
+					if (UAssetPlacementEdMode::DoesPaletteSupportElement(SourceObjectHandle, PlacementSettings->PaletteItems))
+					{
+						TArray<int32> Instances;
+						FSphere SphereToCheck(LastBrushStamp.WorldPosition, LastBrushStamp.Radius);
+						FoliageInfo.Value->GetInstancesInsideSphere(SphereToCheck, Instances);
+						FoliageInfo.Value->RemoveInstances(FoliageActor, Instances, true);
+					}
+				}
+			}
+			else if (ActorSubsystem)
+			{
+				ActorSubsystem->DestroyActor(ActorData->Actor);
+			}
+		}
+	}
 }
diff --git a/Engine/Source/Editor/AssetPlacementEdMode/Private/Tools/PlacementEraseTool.h b/Engine/Source/Editor/AssetPlacementEdMode/Private/Tools/PlacementEraseTool.h
--- a/Engine/Source/Editor/AssetPlacementEdMode/Private/Tools/PlacementEraseTool.h
+++ b/Engine/Source/Editor/AssetPlacementEdMode/Private/Tools/PlacementEraseTool.h
@@ -8,13 +8,12 @@
 #include "PlacementEraseTool.generated.h"
 
 UCLASS(Transient, MinimalAPI)
-class UPlacementModeEraseToolBuilder : public UInteractiveToolBuilder
+class UPlacementModeEraseToolBuilder : public UPlacementToolBuilderBase
 {
 	GENERATED_BODY()
 
-public:
-	virtual bool CanBuildTool(const FToolBuilderState& SceneState) const override;
-	virtual UInteractiveTool* BuildTool(const FToolBuilderState& SceneState) const override;
+protected:
+	virtual UPlacementBrushToolBase* FactoryToolInstance(UObject* Outer) const override;
 };
 
 UCLASS(MinimalAPI)
@@ -23,4 +22,10 @@ class UPlacementModeEraseTool : public UPlacementBrushToolBase
 	GENERATED_BODY()
 public:
 	constexpr static TCHAR ToolName[] = TEXT("EraseTool");
+
+	virtual void OnBeginDrag(const FRay& Ray) override;
+	virtual void OnEndDrag(const FRay& Ray) override;
+
+protected:
+	virtual void OnTick(float DeltaTime) override;
 };
diff --git a/Engine/Source/Editor/AssetPlacementEdMode/Private/Tools/PlacementPlaceTool.cpp b/Engine/Source/Editor/AssetPlacementEdMode/Private/Tools/PlacementPlaceTool.cpp
--- a/Engine/Source/Editor/AssetPlacementEdMode/Private/Tools/PlacementPlaceTool.cpp
+++ b/Engine/Source/Editor/AssetPlacementEdMode/Private/Tools/PlacementPlaceTool.cpp
@@ -5,10 +5,9 @@
 #include "AssetPlacementSettings.h"
 #include "Subsystems/PlacementSubsystem.h"
 #include "Math/UnrealMathUtility.h"
-#include "BaseGizmos/BrushStampIndicator.h"
 #include "Editor.h"
 #include "InstancedFoliage.h"
-#include "ScopedTransaction.h"
+#include "InteractiveToolManager.h"
 
 constexpr TCHAR UPlacementModePlacementTool::ToolName[];
 
@@ -21,12 +20,12 @@ void UPlacementModePlacementTool::OnBeginDrag(const FRay& Ray)
 {
 	Super::OnBeginDrag(Ray);
 
-	TransactionScope = MakeUnique<FScopedTransaction>(NSLOCTEXT("PlacementMode", "PaintAssets", "Paint Asset Stroke"));
+	GetToolManager()->BeginUndoTransaction(NSLOCTEXT("AssetPlacementEdMode", "PaintAssets", "Paint Assets"));
 }
 
 void UPlacementModePlacementTool::OnEndDrag(const FRay& Ray)
 {
-	TransactionScope.Reset();
+	GetToolManager()->EndUndoTransaction();
 
 	Super::OnEndDrag(Ray);
 }
diff --git a/Engine/Source/Editor/AssetPlacementEdMode/Public/AssetPlacementEdMode.h b/Engine/Source/Editor/AssetPlacementEdMode/Public/AssetPlacementEdMode.h
--- a/Engine/Source/Editor/AssetPlacementEdMode/Public/AssetPlacementEdMode.h
+++ b/Engine/Source/Editor/AssetPlacementEdMode/Public/AssetPlacementEdMode.h
@@ -5,6 +5,7 @@
 #include "CoreMinimal.h"
 #include "Tools/LegacyEdModeWidgetHelpers.h"
 #include "Elements/Framework/TypedElementHandle.h"
+#include "AssetPlacementSettings.h"
 
 #include "AssetPlacementEdMode.generated.h"
 
@@ -36,6 +37,8 @@ public:
 	// End of UBaseLegacyWidgetEdMode interface
 	//////////////////
 
+	static bool DoesPaletteSupportElement(const FTypedElementHandle& InElementToCheck, const TArray<FPaletteItem>& InPaletteToCheck);
+
 protected:
 	enum class EPaletteFilter
 	{
diff --git a/Engine/Source/Editor/UnrealEd/Private/Factories/EditorStaticMeshFactory.cpp b/Engine/Source/Editor/UnrealEd/Private/Factories/EditorStaticMeshFactory.cpp
--- a/Engine/Source/Editor/UnrealEd/Private/Factories/EditorStaticMeshFactory.cpp
+++ b/Engine/Source/Editor/UnrealEd/Private/Factories/EditorStaticMeshFactory.cpp
@@ -4,6 +4,7 @@
 
 #include "Elements/Actor/ActorElementData.h"
 #include "Elements/Component/ComponentElementData.h"
+#include "Elements/Object/ObjectElementData.h"
 #include "Elements/Framework/EngineElementsLibrary.h"
 
 #include "ActorPartition/ActorPartitionSubsystem.h"
@@ -89,15 +90,24 @@ FAssetData UEditorStaticMeshFactory::GetAssetDataFromElementHandle(const FTypedE
 		}
 	}
 
+	FAssetData FoundAssetData;
 	if (ISMComponent)
 	{
-		FAssetData ISMStaticMeshAssetData = FAssetData(ISMComponent->GetStaticMesh());
-		if (CanPlaceElementsFromAssetData(ISMStaticMeshAssetData))
+		FoundAssetData = FAssetData(ISMComponent->GetStaticMesh());
+	}
+	else if (const FObjectElementData* ObjectData = InHandle.GetData<FObjectElementData>())
+	{
+		if (ObjectData && ObjectData->Object)
 		{
-			return ISMStaticMeshAssetData;
+			FoundAssetData = FAssetData(ObjectData->Object);
 		}
 	}
 
+	if (CanPlaceElementsFromAssetData(FoundAssetData))
+	{
+		return FoundAssetData;
+	}
+
 	// Todo: deal with instanced handles
 
 	return Super::GetAssetDataFromElementHandle(InHandle);
