Implement clear selection and delete selection commands in placement mode.

# Removed Select All tool - we don't think it makes sense in the open world case, and would be better served by "clean up" type utility to remove all references to an asset over a world.

diff --git a/Engine/Source/Editor/AssetPlacementEdMode/Private/AssetPlacementEdMode.cpp b/Engine/Source/Editor/AssetPlacementEdMode/Private/AssetPlacementEdMode.cpp
--- a/Engine/Source/Editor/AssetPlacementEdMode/Private/AssetPlacementEdMode.cpp
+++ b/Engine/Source/Editor/AssetPlacementEdMode/Private/AssetPlacementEdMode.cpp
@@ -8,6 +8,7 @@
 #include "EdModeInteractiveToolsContext.h"
 #include "EditorModeManager.h"
 #include "Selection.h"
+#include "EngineUtils.h"
 
 #include "Tools/PlacementSelectTool.h"
 #include "Tools/PlacementLassoSelectTool.h"
@@ -19,9 +20,15 @@
 #include "Factories/AssetFactoryInterface.h"
 #include "Elements/Actor/ActorElementData.h"
 #include "Elements/Framework/EngineElementsLibrary.h"
+#include "Elements/Framework/TypedElementRegistry.h"
+#include "Elements/Interfaces/TypedElementAssetDataInterface.h"
+#include "Elements/Interfaces/TypedElementObjectInterface.h"
 
 #include "Settings/LevelEditorMiscSettings.h"
 
+#include "Subsystems/EditorActorSubsystem.h"
+#include "InstancedFoliageActor.h"
+
 #define LOCTEXT_NAMESPACE "AssetPlacementEdMode"
 
 constexpr TCHAR UAssetPlacementEdMode::AssetPlacementEdModeID[];
@@ -45,24 +52,24 @@ void UAssetPlacementEdMode::Enter()
 {
 	Super::Enter();
 
-	UAssetPlacementSettings* PlacementSettings = Cast<UAssetPlacementSettings>(SettingsObject);
+	SettingsObjectAsPlacementSettings = Cast<UAssetPlacementSettings>(SettingsObject);
 
 	const FAssetPlacementEdModeCommands& PlacementModeCommands = FAssetPlacementEdModeCommands::Get();
 	RegisterTool(PlacementModeCommands.Select, UPlacementModeSelectTool::ToolName, NewObject<UPlacementModeSelectToolBuilder>(this));
 	RegisterTool(PlacementModeCommands.LassoSelect, UPlacementModeLassoSelectTool::ToolName, NewObject<UPlacementModeLassoSelectToolBuilder>(this));
 
 	UPlacementModePlacementToolBuilder* PlaceToolBuilder = NewObject<UPlacementModePlacementToolBuilder>(this);
-	PlaceToolBuilder->PlacementSettings = Cast<UAssetPlacementSettings>(SettingsObject);
+	PlaceToolBuilder->PlacementSettings = SettingsObjectAsPlacementSettings;
 	RegisterTool(PlacementModeCommands.Place, UPlacementModePlacementTool::ToolName, PlaceToolBuilder);
 
 	UPlacementModePlaceSingleToolBuilder* SinglePlaceToolBuilder = NewObject<UPlacementModePlaceSingleToolBuilder>(this);
-	SinglePlaceToolBuilder->PlacementSettings = Cast<UAssetPlacementSettings>(SettingsObject);
+	SinglePlaceToolBuilder->PlacementSettings = SettingsObjectAsPlacementSettings;
 	RegisterTool(PlacementModeCommands.PlaceSingle, UPlacementModePlaceSingleTool::ToolName, SinglePlaceToolBuilder);
 
 	RegisterTool(PlacementModeCommands.Fill, UPlacementModeFillTool::ToolName, NewObject<UPlacementModeFillToolBuilder>(this));
 
 	UPlacementModeEraseToolBuilder* EraseToolBuilder = NewObject<UPlacementModeEraseToolBuilder>(this);
-	EraseToolBuilder->PlacementSettings = Cast<UAssetPlacementSettings>(SettingsObject);
+	EraseToolBuilder->PlacementSettings = SettingsObjectAsPlacementSettings;
 	RegisterTool(PlacementModeCommands.Erase, UPlacementModeEraseTool::ToolName, EraseToolBuilder);
 
 	ToolsContext->ToolManager->SelectActiveToolType(EToolSide::Mouse, UPlacementModeSelectTool::ToolName);
@@ -73,9 +80,16 @@ void UAssetPlacementEdMode::Enter()
 	Owner->GetSelectedComponents()->GetElementSelectionSet()->ClearSelection(FTypedElementSelectionOptions());
 }
 
+void UAssetPlacementEdMode::Exit()
+{
+	SettingsObjectAsPlacementSettings.Reset();
+
+	Super::Exit();
+}
+
 void UAssetPlacementEdMode::CreateToolkit()
 {
-	Toolkit = MakeShareable(new FAssetPlacementEdModeToolkit(Cast<UAssetPlacementSettings>(SettingsObject)));
+	Toolkit = MakeShareable(new FAssetPlacementEdModeToolkit(SettingsObjectAsPlacementSettings));
 }
 
 TMap<FName, TArray<TSharedPtr<FUICommandInfo>>> UAssetPlacementEdMode::GetModeCommands() const
@@ -88,17 +102,13 @@ void UAssetPlacementEdMode::BindCommands()
 	const TSharedRef<FUICommandList>& CommandList = Toolkit->GetToolkitCommands();
 	const FAssetPlacementEdModeCommands& PlacementModeCommands = FAssetPlacementEdModeCommands::Get();
 
-	CommandList->MapAction(PlacementModeCommands.SelectAll,
-		FExecuteAction::CreateUObject(this, &UAssetPlacementEdMode::SelectAssets, EPaletteFilter::ActivePaletteOnly, ESelectMode::Select),
-		FCanExecuteAction::CreateUObject(this, &UAssetPlacementEdMode::HasAnyAssetsInPalette, EPaletteFilter::ActivePaletteOnly));
-
 	CommandList->MapAction(PlacementModeCommands.Deselect,
-		FExecuteAction::CreateUObject(this, &UAssetPlacementEdMode::SelectAssets, EPaletteFilter::ActivePaletteOnly, ESelectMode::Deselect),
-		FCanExecuteAction::CreateUObject(this, &UAssetPlacementEdMode::HasAnyAssetsInPalette, EPaletteFilter::ActivePaletteOnly));
+		FExecuteAction::CreateUObject(this, &UAssetPlacementEdMode::ClearSelection),
+		FCanExecuteAction::CreateUObject(this, &UAssetPlacementEdMode::HasActiveSelection));
 
 	CommandList->MapAction(PlacementModeCommands.Delete,
-		FExecuteAction::CreateUObject(this, &UAssetPlacementEdMode::DeleteAssets),
-		FCanExecuteAction::CreateUObject(this, &UAssetPlacementEdMode::HasAnyAssetsInPalette, EPaletteFilter::ActivePaletteOnly));
+		FExecuteAction::CreateUObject(this, &UAssetPlacementEdMode::DeleteSelection),
+		FCanExecuteAction::CreateUObject(this, &UAssetPlacementEdMode::HasActiveSelection));
 }
 
 bool UAssetPlacementEdMode::IsSelectionAllowed(AActor* InActor, bool bInSelection) const
@@ -109,10 +119,8 @@ bool UAssetPlacementEdMode::IsSelectionAllowed(AActor* InActor, bool bInSelectio
 		return false;
 	}
 
-	UAssetPlacementSettings* PlacementSettings = Cast<UAssetPlacementSettings>(SettingsObject);
 	FTypedElementHandle ActorHandle = UEngineElementsLibrary::AcquireEditorActorElementHandle(InActor);
-
-	return DoesPaletteSupportElement(ActorHandle, PlacementSettings->PaletteItems);
+	return DoesPaletteSupportElement(ActorHandle, SettingsObjectAsPlacementSettings->PaletteItems);
 }
 
 bool UAssetPlacementEdMode::UsesPropertyWidgets() const
@@ -142,11 +150,20 @@ bool UAssetPlacementEdMode::IsEnabled()
 
 bool UAssetPlacementEdMode::DoesPaletteSupportElement(const FTypedElementHandle& InElementToCheck, const TArray<FPaletteItem>& InPaletteToCheck)
 {
-	if (!InElementToCheck.IsSet())
+	if (TTypedElement<UTypedElementAssetDataInterface> AssetDataInterface = UTypedElementRegistry::GetInstance()->GetElement<UTypedElementAssetDataInterface>(InElementToCheck))
 	{
-		return false;
+		TArray<FAssetData> ReferencedAssetDatas = AssetDataInterface.GetAllReferencedAssetDatas();
+		for (const FPaletteItem& Item : InPaletteToCheck)
+		{
+			if (ReferencedAssetDatas.Find(Item.AssetData) != INDEX_NONE)
+			{
+				return true;
+			}
+		}
 	}
 
+	// The current implementation of the asset data interface for actors requires that individual actors report on assets contained within their components.
+	// Not all actors do this reliably, so additionally check the supplied factory for a match. 
 	for (const FPaletteItem& Item : InPaletteToCheck)
 	{
 		FAssetData FoundAssetDataFromFactory = Item.FactoryOverride->GetAssetDataFromElementHandle(InElementToCheck);
@@ -159,22 +176,82 @@ bool UAssetPlacementEdMode::DoesPaletteSupportElement(const FTypedElementHandle&
 	return false;
 }
 
-void UAssetPlacementEdMode::SelectAssets(EPaletteFilter InSelectAllType, ESelectMode InSelectMode)
+void UAssetPlacementEdMode::DeleteSelection()
 {
+	GetToolManager()->BeginUndoTransaction(LOCTEXT("PlacementDeleteAllSelected", "Delete Selected Assets"));
+
+	// Todo - replace with delete in world interface and replace foliage with element handles
+	UEditorActorSubsystem* ActorSubsystem = GEditor->GetEditorSubsystem<UEditorActorSubsystem>();
+	auto DeleteSelectedElement = [ActorSubsystem](const TTypedElement<UTypedElementObjectInterface>& InElementInterface)
+	{
+		if (InElementInterface.IsSet())
+		{
+			if (AActor* Actor = InElementInterface.GetObjectAs<AActor>())
+			{
+				if (ActorSubsystem)
+				{
+					ActorSubsystem->DestroyActor(Actor);
+				}
+			}
+		}
+	};
+
+	// Gather a copy of the selected handles, since the delete operation will remove them from the selection set.
+	TArray<FTypedElementHandle> SelectedElementHandles;
+	Owner->GetEditorSelectionSet()->GetSelectedElementHandles(SelectedElementHandles, UTypedElementObjectInterface::StaticClass());
+	for (const FTypedElementHandle& ElementHandle : SelectedElementHandles)
+	{
+		if (TTypedElement<UTypedElementObjectInterface> ObjectInterface = UTypedElementRegistry::GetInstance()->GetElement<UTypedElementObjectInterface>(ElementHandle))
+		{
+			DeleteSelectedElement(ObjectInterface);
+		}
+	}
+
+	for (TActorIterator<AInstancedFoliageActor> It(GetWorld()); It; ++It)
+	{
+		if (AInstancedFoliageActor* FoliageActor = *It)
+		{
+			for (auto& FoliageInfo : FoliageActor->FoliageInfos)
+			{
+				TArray<int32> SelectedIndices = FoliageInfo.Value->SelectedIndices.Array();
+				FoliageInfo.Value->RemoveInstances(FoliageActor, SelectedIndices, true);
+			}
+		}
+	}
+
+	GetToolManager()->EndUndoTransaction();
 }
 
-void UAssetPlacementEdMode::DeleteAssets()
+void UAssetPlacementEdMode::ClearSelection()
 {
+	GetToolManager()->BeginUndoTransaction(LOCTEXT("PlacementClearSelection", "Clear Selection"));
+	if (UTypedElementSelectionSet* SelectionSet = Owner->GetEditorSelectionSet())
+	{
+		SelectionSet->ClearSelection(FTypedElementSelectionOptions());
+	}
+
+	// Todo - remove with instanced handles
+	for (TActorIterator<AInstancedFoliageActor> It(GetWorld()); It; ++It)
+	{
+		if (AInstancedFoliageActor* FoliageActor = *It)
+		{
+			for (auto& FoliageInfo : FoliageActor->FoliageInfos)
+			{
+				FoliageInfo.Value->ClearSelection();
+			}
+		}
+	}
+	GetToolManager()->EndUndoTransaction();
 }
 
-void UAssetPlacementEdMode::MoveAssetToActivePartition()
+bool UAssetPlacementEdMode::HasAnyAssetsInPalette() const
 {
+	return (SettingsObjectAsPlacementSettings->PaletteItems.Num() > 0);
 }
 
-bool UAssetPlacementEdMode::HasAnyAssetsInPalette(EPaletteFilter InSelectAllType)
+bool UAssetPlacementEdMode::HasActiveSelection() const
 {
-	UAssetPlacementSettings* PlacementSettings = Cast<UAssetPlacementSettings>(SettingsObject);
-	return (PlacementSettings->PaletteItems.Num() > 0);
+	return HasAnyAssetsInPalette() && Owner->GetEditorSelectionSet()->HasSelectedElements();
 }
 
 #undef LOCTEXT_NAMESPACE
diff --git a/Engine/Source/Editor/AssetPlacementEdMode/Private/AssetPlacementEdModeCommands.cpp b/Engine/Source/Editor/AssetPlacementEdMode/Private/AssetPlacementEdModeCommands.cpp
--- a/Engine/Source/Editor/AssetPlacementEdMode/Private/AssetPlacementEdModeCommands.cpp
+++ b/Engine/Source/Editor/AssetPlacementEdMode/Private/AssetPlacementEdModeCommands.cpp
@@ -18,21 +18,19 @@ void FAssetPlacementEdModeCommands::RegisterCommands()
 	TArray <TSharedPtr<FUICommandInfo>>& ToolCommands = Commands.FindOrAdd(NAME_Default);
 	UI_COMMAND(Select, "Select", "Select", EUserInterfaceActionType::ToggleButton, FInputChord());
 	ToolCommands.Add(Select);
-	UI_COMMAND(SelectAll, "All", "Select All Assets from the Palette", EUserInterfaceActionType::ToggleButton, FInputChord());
-	ToolCommands.Add(SelectAll);
-	UI_COMMAND(Deselect, "Deselect", "Deselect All Assets from the Palette", EUserInterfaceActionType::ToggleButton, FInputChord());
+	UI_COMMAND(Deselect, "Deselect", "Clears the Current Selection.", EUserInterfaceActionType::ToggleButton, FInputChord());
 	ToolCommands.Add(Deselect);
-	UI_COMMAND(LassoSelect, "Lasso", "Select Assets with Lasso", EUserInterfaceActionType::ToggleButton, FInputChord());
+	UI_COMMAND(LassoSelect, "Lasso", "Select Assets with a Brush Tool.", EUserInterfaceActionType::ToggleButton, FInputChord());
 	ToolCommands.Add(LassoSelect);
-	UI_COMMAND(Place, "Place", "Place Assets from the Palette", EUserInterfaceActionType::ToggleButton, FInputChord());
+	UI_COMMAND(Place, "Place", "Place Assets from the Palette.", EUserInterfaceActionType::ToggleButton, FInputChord());
 	ToolCommands.Add(Place);
-	UI_COMMAND(PlaceSingle, "Single", "Place a Single Asset from the Palette", EUserInterfaceActionType::ToggleButton, FInputChord());
+	UI_COMMAND(PlaceSingle, "Single", "Place a Single Asset from the Palette.", EUserInterfaceActionType::ToggleButton, FInputChord());
 	ToolCommands.Add(PlaceSingle);
-	UI_COMMAND(Fill, "Fill", "Fill the Selected Target with Assets from the Palette", EUserInterfaceActionType::ToggleButton, FInputChord());
+	UI_COMMAND(Fill, "Fill", "Fill the Selected Target with Assets from the Palette.", EUserInterfaceActionType::ToggleButton, FInputChord());
 	ToolCommands.Add(Fill);
-	UI_COMMAND(Erase, "Erase", "Erases Assets Matching the Palette", EUserInterfaceActionType::ToggleButton, FInputChord());
+	UI_COMMAND(Erase, "Erase", "Erases Assets Matching the Palette.", EUserInterfaceActionType::ToggleButton, FInputChord());
 	ToolCommands.Add(Erase);
-	UI_COMMAND(Delete, "Delete", "Remove the Selected Assets", EUserInterfaceActionType::ToggleButton, FInputChord());
+	UI_COMMAND(Delete, "Delete", "Deletes Assets in the Current Selection.", EUserInterfaceActionType::ToggleButton, FInputChord());
 	ToolCommands.Add(Delete);
 }
 
diff --git a/Engine/Source/Editor/AssetPlacementEdMode/Private/Tools/PlacementBrushToolBase.cpp b/Engine/Source/Editor/AssetPlacementEdMode/Private/Tools/PlacementBrushToolBase.cpp
--- a/Engine/Source/Editor/AssetPlacementEdMode/Private/Tools/PlacementBrushToolBase.cpp
+++ b/Engine/Source/Editor/AssetPlacementEdMode/Private/Tools/PlacementBrushToolBase.cpp
@@ -20,7 +20,7 @@
 
 bool UPlacementToolBuilderBase::CanBuildTool(const FToolBuilderState& SceneState) const
 {
-	return PlacementSettings && PlacementSettings->PaletteItems.Num();
+	return PlacementSettings.IsValid() && PlacementSettings->PaletteItems.Num();
 }
 
 UInteractiveTool* UPlacementToolBuilderBase::BuildTool(const FToolBuilderState& SceneState) const
@@ -170,7 +170,7 @@ TArray<FTypedElementHandle> UPlacementBrushToolBase::GetElementsInBrushRadius()
 			continue;
 		}
 
-		// Don't place foliage on invisible walls / triggers / volumes
+		// Don't place assets on invisible walls / triggers / volumes
 		if (HitComponent->IsA<UBrushComponent>())
 		{
 			continue;
diff --git a/Engine/Source/Editor/AssetPlacementEdMode/Private/Tools/PlacementBrushToolBase.h b/Engine/Source/Editor/AssetPlacementEdMode/Private/Tools/PlacementBrushToolBase.h
--- a/Engine/Source/Editor/AssetPlacementEdMode/Private/Tools/PlacementBrushToolBase.h
+++ b/Engine/Source/Editor/AssetPlacementEdMode/Private/Tools/PlacementBrushToolBase.h
@@ -21,7 +21,7 @@ public:
 	virtual UInteractiveTool* BuildTool(const FToolBuilderState& SceneState) const override;
 
 	UPROPERTY()
-	TObjectPtr<UAssetPlacementSettings> PlacementSettings;
+	TWeakObjectPtr<UAssetPlacementSettings> PlacementSettings;
 
 protected:
 	virtual UPlacementBrushToolBase* FactoryToolInstance(UObject* Outer) const PURE_VIRTUAL(UPlacementToolBuilderBase::FactoryToolInstance, return nullptr; );
diff --git a/Engine/Source/Editor/AssetPlacementEdMode/Private/Tools/PlacementPlaceTool.h b/Engine/Source/Editor/AssetPlacementEdMode/Private/Tools/PlacementPlaceTool.h
--- a/Engine/Source/Editor/AssetPlacementEdMode/Private/Tools/PlacementPlaceTool.h
+++ b/Engine/Source/Editor/AssetPlacementEdMode/Private/Tools/PlacementPlaceTool.h
@@ -7,8 +7,6 @@
 
 #include "PlacementPlaceTool.generated.h"
 
-class FScopedTransaction;
-
 UCLASS(Transient, MinimalAPI)
 class UPlacementModePlacementToolBuilder : public UPlacementToolBuilderBase
 {
@@ -32,6 +30,4 @@ protected:
 	virtual void OnTick(float DeltaTime) override;
 
 	void GetRandomVectorInBrush(FVector& OutStart, FVector& OutEnd);
-
-	TUniquePtr<FScopedTransaction> TransactionScope;
 };
diff --git a/Engine/Source/Editor/AssetPlacementEdMode/Public/AssetPlacementEdMode.h b/Engine/Source/Editor/AssetPlacementEdMode/Public/AssetPlacementEdMode.h
--- a/Engine/Source/Editor/AssetPlacementEdMode/Public/AssetPlacementEdMode.h
+++ b/Engine/Source/Editor/AssetPlacementEdMode/Public/AssetPlacementEdMode.h
@@ -23,6 +23,7 @@ public:
 	////////////////
 	// UEdMode interface
 	virtual void Enter() override;
+	virtual void Exit() override;
 	virtual void CreateToolkit() override;
 	virtual TMap<FName, TArray<TSharedPtr<FUICommandInfo>>> GetModeCommands() const override;
 	virtual void BindCommands() override;
@@ -41,20 +42,16 @@ public:
 	static bool DoesPaletteSupportElement(const FTypedElementHandle& InElementToCheck, const TArray<FPaletteItem>& InPaletteToCheck);
 
 protected:
-	enum class EPaletteFilter
-	{
-		EntirePalette,
-		ActivePaletteOnly,
-	};
-
 	enum class ESelectMode
 	{
 		Select,
 		Deselect,
 	};
 
-	void SelectAssets(EPaletteFilter InSelectAllType, ESelectMode InSelectMode);
-	void DeleteAssets();
-	void MoveAssetToActivePartition();
-	bool HasAnyAssetsInPalette(EPaletteFilter InSelectAllType);
+	void DeleteSelection();
+	void ClearSelection();
+	bool HasAnyAssetsInPalette() const;
+	bool HasActiveSelection() const;
+
+	TWeakObjectPtr<UAssetPlacementSettings> SettingsObjectAsPlacementSettings;
 };
diff --git a/Engine/Source/Editor/UnrealEd/Private/EditorModeManager.cpp b/Engine/Source/Editor/UnrealEd/Private/EditorModeManager.cpp
--- a/Engine/Source/Editor/UnrealEd/Private/EditorModeManager.cpp
+++ b/Engine/Source/Editor/UnrealEd/Private/EditorModeManager.cpp
@@ -34,6 +34,7 @@
 #include "InteractiveGizmoManager.h"
 #include "EdModeInteractiveToolsContext.h"
 #include "Tools/LegacyEdModeInterfaces.h"
+#include "Elements/Framework/TypedElementSelectionSet.h"
 
 /*------------------------------------------------------------------------------
 	FEditorModeTools.
@@ -159,6 +160,11 @@ USelection* FEditorModeTools::GetSelectedComponents() const
 	return GEditor->GetSelectedComponents();
 }
 
+UTypedElementSelectionSet* FEditorModeTools::GetEditorSelectionSet() const
+{
+	return GetSelectedActors()->GetElementSelectionSet();
+}
+
 UWorld* FEditorModeTools::GetWorld() const
 {
 	// When in 'Simulate' mode, the editor mode tools will actually interact with the PIE world
@@ -553,6 +559,7 @@ void FEditorModeTools::DeactivateScriptableModeAtIndex(int32 InIndex)
 	check(InIndex >= 0 && InIndex < ActiveScriptableModes.Num());
 
 	UEdMode* Mode = ActiveScriptableModes[InIndex];
+	ActiveScriptableModes.RemoveAt(InIndex);
 
 	Mode->Exit();
 
@@ -570,7 +577,6 @@ void FEditorModeTools::DeactivateScriptableModeAtIndex(int32 InIndex)
 	RebuildModeToolBar();
 
 	RecycledScriptableModes.Add(Mode->GetID(), Mode);
-	ActiveScriptableModes.RemoveAt(InIndex);
 }
 
 void FEditorModeTools::OnModeUnregistered(FEditorModeID ModeID)
diff --git a/Engine/Source/Editor/UnrealEd/Public/EditorModeManager.h b/Engine/Source/Editor/UnrealEd/Public/EditorModeManager.h
--- a/Engine/Source/Editor/UnrealEd/Public/EditorModeManager.h
+++ b/Engine/Source/Editor/UnrealEd/Public/EditorModeManager.h
@@ -26,6 +26,7 @@ class UEdMode;
 class UInteractiveGizmoManager;
 class UInputRouter;
 class UEdModeInteractiveToolsContext;
+class UTypedElementSelectionSet;
 
 /**
  * A helper class to store the state of the various editor modes.
@@ -445,6 +446,12 @@ public:
 	 */
 	virtual USelection* GetSelectedComponents() const;
 
+	/**
+	 * Returns the selection set for the toolkit host.
+	 * (i.e. the selection set for the level editor)
+	 */
+	UTypedElementSelectionSet* GetEditorSelectionSet() const;
+
 	/**
 	 * Returns the world that is being edited by this mode manager
 	 */ 
diff --git a/Engine/Source/Runtime/Engine/Private/Animation/SkeletalMeshActor.cpp b/Engine/Source/Runtime/Engine/Private/Animation/SkeletalMeshActor.cpp
--- a/Engine/Source/Runtime/Engine/Private/Animation/SkeletalMeshActor.cpp
+++ b/Engine/Source/Runtime/Engine/Private/Animation/SkeletalMeshActor.cpp
@@ -257,6 +257,12 @@ bool ASkeletalMeshActor::GetReferencedContentObjects(TArray<UObject*>& Objects)
 	{
 		Objects.Add(SkeletalMeshComponent->SkeletalMesh);
 	}
+
+	if (SkeletalMeshComponent->GetAnimationMode() == EAnimationMode::Type::AnimationSingleNode)
+	{
+		Objects.Add(SkeletalMeshComponent->AnimationData.AnimToPlay);
+	}
+
 	return true;
 }
 
diff --git a/Engine/Source/Runtime/TypedElementRuntime/Public/Elements/Interfaces/TypedElementAssetDataInterface.h b/Engine/Source/Runtime/TypedElementRuntime/Public/Elements/Interfaces/TypedElementAssetDataInterface.h
--- a/Engine/Source/Runtime/TypedElementRuntime/Public/Elements/Interfaces/TypedElementAssetDataInterface.h
+++ b/Engine/Source/Runtime/TypedElementRuntime/Public/Elements/Interfaces/TypedElementAssetDataInterface.h
@@ -33,5 +33,5 @@ template <>
 struct TTypedElement<UTypedElementAssetDataInterface> : public TTypedElementBase<UTypedElementAssetDataInterface>
 {
 	TArray<FAssetData> GetAllReferencedAssetDatas() const { return InterfacePtr->GetAllReferencedAssetDatas(*this); }
-	FAssetData GetAssetData() { return InterfacePtr->GetAssetData(*this); }
+	FAssetData GetAssetData() const { return InterfacePtr->GetAssetData(*this); }
 };
diff --git a/Engine/Source/Runtime/TypedElementRuntime/Public/Elements/Interfaces/TypedElementObjectInterface.h b/Engine/Source/Runtime/TypedElementRuntime/Public/Elements/Interfaces/TypedElementObjectInterface.h
--- a/Engine/Source/Runtime/TypedElementRuntime/Public/Elements/Interfaces/TypedElementObjectInterface.h
+++ b/Engine/Source/Runtime/TypedElementRuntime/Public/Elements/Interfaces/TypedElementObjectInterface.h
@@ -62,11 +62,11 @@ template <>
 struct TTypedElement<UTypedElementObjectInterface> : public TTypedElementBase<UTypedElementObjectInterface>
 {
 	UObject* GetObject() const { return InterfacePtr->GetObject(*this); }
-	UClass* GetObjectClass() { return InterfacePtr->GetObjectClass(*this); }
+	UClass* GetObjectClass() const { return InterfacePtr->GetObjectClass(*this); }
 
 	template <class CastTo>
-	CastTo* GetObjectAs() { return InterfacePtr->GetObjectAs<CastTo>(*this); }
+	CastTo* GetObjectAs() const { return InterfacePtr->GetObjectAs<CastTo>(*this); }
 
 	template <class CastTo>
-	CastTo* GetObjectAs(TSubclassOf<CastTo> TargetClass) { return InterfacePtr->GetObjectAs<CastTo>(*this, TargetClass); }
+	CastTo* GetObjectAs(TSubclassOf<CastTo> TargetClass) const { return InterfacePtr->GetObjectAs<CastTo>(*this, TargetClass); }
 };
