Fix a crash when swapping to list view for palette, and reposition the palette to the bottom of the mode toolbox.

diff --git a/Engine/Source/Editor/AssetPlacementEdMode/Private/AssetPlacementEdModeToolkit.cpp b/Engine/Source/Editor/AssetPlacementEdMode/Private/AssetPlacementEdModeToolkit.cpp
--- a/Engine/Source/Editor/AssetPlacementEdMode/Private/AssetPlacementEdModeToolkit.cpp
+++ b/Engine/Source/Editor/AssetPlacementEdMode/Private/AssetPlacementEdModeToolkit.cpp
@@ -38,6 +38,12 @@ TSharedPtr<SWidget> FAssetPlacementEdModeToolkit::GetInlineContent() const
 {
 	return SNew(SVerticalBox)
 		+ SVerticalBox::Slot()
+		.AutoHeight()
+		[
+			FModeToolkit::GetInlineContent().ToSharedRef()
+		]
+		+ SVerticalBox::Slot()
+		.AutoHeight()
 		[
 			SNew(SExpandableArea)
 			.AreaTitle(LOCTEXT("AssetPaletteHeader", "Asset Palette"))
@@ -46,10 +52,6 @@ TSharedPtr<SWidget> FAssetPlacementEdModeToolkit::GetInlineContent() const
 				SNew(SAssetPlacementPalette)
 				.PlacementSettings(PlacementSettings)
 			]
-		]
-		+ SVerticalBox::Slot()
-		[
-			FModeToolkit::GetInlineContent().ToSharedRef()
 		];
 }
 
diff --git a/Engine/Source/Editor/AssetPlacementEdMode/Private/SAssetPlacementPalette.cpp b/Engine/Source/Editor/AssetPlacementEdMode/Private/SAssetPlacementPalette.cpp
--- a/Engine/Source/Editor/AssetPlacementEdMode/Private/SAssetPlacementPalette.cpp
+++ b/Engine/Source/Editor/AssetPlacementEdMode/Private/SAssetPlacementPalette.cpp
@@ -184,9 +184,9 @@ void SAssetPlacementPalette::Construct(const FArguments& InArgs)
 	ChildSlot
 	[
 		SNew(SVerticalBox)
-
 		+ SVerticalBox::Slot()
 		.HAlign(HAlign_Fill)
+		.AutoHeight()
 		[
 			// Top bar
 			SNew(SBorder)
@@ -226,43 +226,8 @@ void SAssetPlacementPalette::Construct(const FArguments& InArgs)
 						]
 					]
 				]
-
-				+ SVerticalBox::Slot()
-				.VAlign(VAlign_Center)
-				[
-					SNew(SHorizontalBox)
-					+SHorizontalBox::Slot()
-					.FillWidth(.75f)
-					[
-						// +Add Placement Type button
-						SAssignNew(AddPlacementTypeCombo, SComboButton)
-						.HAlign(HAlign_Center)
-						.ForegroundColor(FLinearColor::White)
-						.ButtonStyle(FAppStyle::Get(), "FlatButton.Success")
-						.OnGetMenuContent(this, &SAssetPlacementPalette::GetAddPlacementTypePicker)
-						.ButtonContent()
-							[
-								SNew(SHorizontalBox)
-								+ SHorizontalBox::Slot()
-								.VAlign(VAlign_Center)
-								.AutoWidth()
-								[
-									SNew(STextBlock)
-									.TextStyle(FAssetPlacementEdModeStyle::Get(), "AssetPlacementEdMode.AddAssetType.Text")
-									.Text(FText::FromString(FString(TEXT("\xf067"))) /*fa-plus*/)
-								]
-								+ SHorizontalBox::Slot()
-								.VAlign(VAlign_Center)
-								[
-									SNew(STextBlock)
-									.Text(LOCTEXT("AddPlacementTypeButtonLabel", "Add Asset Type"))
-									.TextStyle(FAssetPlacementEdModeStyle::Get(), "AssetPlacementEdMode.AddAssetType.Text")
-								]
-							]
-						]
-					]
-				]
 			]
+		]
 
 		+ SVerticalBox::Slot()
 		[
@@ -456,6 +421,7 @@ TSharedRef<SWidgetSwitcher> SAssetPlacementPalette::CreatePaletteViews()
 		.TreeItemsSource(&FilteredItems)
 		.SelectionMode(ESelectionMode::Multi)
 		.OnGenerateRow(this, &SAssetPlacementPalette::TreeViewGenerateRow)
+		.OnGetChildren(this, &SAssetPlacementPalette::TreeViewGetChildren)
 		.OnContextMenuOpening(this, &SAssetPlacementPalette::ConstructPlacementTypeContextMenu)
 		.OnSelectionChanged(this, &SAssetPlacementPalette::OnSelectionChanged)
 		.OnMouseButtonDoubleClick(this, &SAssetPlacementPalette::OnItemDoubleClicked)
@@ -963,6 +929,11 @@ TSharedRef<ITableRow> SAssetPlacementPalette::TreeViewGenerateRow(FPlacementPale
 	return SNew(SAssetPlacementPaletteItemRow, OwnerTable, Item);
 }
 
+void SAssetPlacementPalette::TreeViewGetChildren(FPlacementPaletteItemModelPtr Item, TArray<FPlacementPaletteItemModelPtr>& OutChildren)
+{
+	// Items do not have any children
+}
+
 ECheckBoxState SAssetPlacementPalette::GetState_AllMeshes() const
 {
 	bool bHasChecked = false;
diff --git a/Engine/Source/Editor/AssetPlacementEdMode/Private/SAssetPlacementPalette.h b/Engine/Source/Editor/AssetPlacementEdMode/Private/SAssetPlacementPalette.h
--- a/Engine/Source/Editor/AssetPlacementEdMode/Private/SAssetPlacementPalette.h
+++ b/Engine/Source/Editor/AssetPlacementEdMode/Private/SAssetPlacementPalette.h
@@ -199,6 +199,7 @@ private:	// TREE VIEW
 
 	/** Generates a row widget for Placement mesh item */
 	TSharedRef<ITableRow> TreeViewGenerateRow(FPlacementPaletteItemModelPtr Item, const TSharedRef<STableViewBase>& OwnerTable);
+	void TreeViewGetChildren(FPlacementPaletteItemModelPtr Item, TArray<FPlacementPaletteItemModelPtr>& OutChildren);
 
 	/** Text for Placement meshes list header */
 	FText GetTypeColumnHeaderText() const;
