Disable Fill tool in placement mode.

# The implementation of this is backlogged due to bookkeeping information for arbitrary (non-foliage) types.

diff --git a/Engine/Source/Editor/AssetPlacementEdMode/Private/AssetPlacementEdMode.cpp b/Engine/Source/Editor/AssetPlacementEdMode/Private/AssetPlacementEdMode.cpp
--- a/Engine/Source/Editor/AssetPlacementEdMode/Private/AssetPlacementEdMode.cpp
+++ b/Engine/Source/Editor/AssetPlacementEdMode/Private/AssetPlacementEdMode.cpp
@@ -15,7 +15,6 @@
 #include "Tools/PlacementLassoSelectTool.h"
 #include "Tools/PlacementPlaceTool.h"
 #include "Tools/PlacementPlaceSingleTool.h"
-#include "Tools/PlacementFillTool.h"
 #include "Tools/PlacementEraseTool.h"
 
 #include "Factories/AssetFactoryInterface.h"
@@ -69,8 +68,6 @@ void UAssetPlacementEdMode::Enter()
 	SinglePlaceToolBuilder->PlacementSettings = SettingsObjectAsPlacementSettings;
 	RegisterTool(PlacementModeCommands.PlaceSingle, UPlacementModePlaceSingleTool::ToolName, SinglePlaceToolBuilder);
 
-	RegisterTool(PlacementModeCommands.Fill, UPlacementModeFillTool::ToolName, NewObject<UPlacementModeFillToolBuilder>(this));
-
 	UPlacementModeEraseToolBuilder* EraseToolBuilder = NewObject<UPlacementModeEraseToolBuilder>(this);
 	EraseToolBuilder->PlacementSettings = SettingsObjectAsPlacementSettings;
 	RegisterTool(PlacementModeCommands.Erase, UPlacementModeEraseTool::ToolName, EraseToolBuilder);
diff --git a/Engine/Source/Editor/AssetPlacementEdMode/Private/AssetPlacementEdModeCommands.cpp b/Engine/Source/Editor/AssetPlacementEdMode/Private/AssetPlacementEdModeCommands.cpp
--- a/Engine/Source/Editor/AssetPlacementEdMode/Private/AssetPlacementEdModeCommands.cpp
+++ b/Engine/Source/Editor/AssetPlacementEdMode/Private/AssetPlacementEdModeCommands.cpp
@@ -26,8 +26,6 @@ void FAssetPlacementEdModeCommands::RegisterCommands()
 	ToolCommands.Add(Place);
 	UI_COMMAND(PlaceSingle, "Single", "Place a Single Asset from the Palette.", EUserInterfaceActionType::ToggleButton, FInputChord());
 	ToolCommands.Add(PlaceSingle);
-	UI_COMMAND(Fill, "Fill", "Fill the Selected Target with Assets from the Palette.", EUserInterfaceActionType::ToggleButton, FInputChord());
-	ToolCommands.Add(Fill);
 	UI_COMMAND(Erase, "Erase", "Erases Assets Matching the Palette.", EUserInterfaceActionType::ToggleButton, FInputChord());
 	ToolCommands.Add(Erase);
 	UI_COMMAND(Delete, "Delete", "Deletes Assets in the Current Selection.", EUserInterfaceActionType::ToggleButton, FInputChord());
diff --git a/Engine/Source/Editor/AssetPlacementEdMode/Public/AssetPlacementEdMode.h b/Engine/Source/Editor/AssetPlacementEdMode/Public/AssetPlacementEdMode.h
--- a/Engine/Source/Editor/AssetPlacementEdMode/Public/AssetPlacementEdMode.h
+++ b/Engine/Source/Editor/AssetPlacementEdMode/Public/AssetPlacementEdMode.h
@@ -9,8 +9,6 @@
 
 #include "AssetPlacementEdMode.generated.h"
 
-class UTypedElementSelectionSet;
-
 UCLASS()
 class UAssetPlacementEdMode : public UBaseLegacyWidgetEdMode
 {
diff --git a/Engine/Source/Editor/AssetPlacementEdMode/Public/AssetPlacementEdModeCommands.h b/Engine/Source/Editor/AssetPlacementEdMode/Public/AssetPlacementEdModeCommands.h
--- a/Engine/Source/Editor/AssetPlacementEdMode/Public/AssetPlacementEdModeCommands.h
+++ b/Engine/Source/Editor/AssetPlacementEdMode/Public/AssetPlacementEdModeCommands.h
@@ -19,7 +19,6 @@ public:
 	TSharedPtr<FUICommandInfo> LassoSelect;
 	TSharedPtr<FUICommandInfo> Place;
 	TSharedPtr<FUICommandInfo> PlaceSingle;
-	TSharedPtr<FUICommandInfo> Fill;
 	TSharedPtr<FUICommandInfo> Erase;
 	TSharedPtr<FUICommandInfo> Delete;
 
