UX pass on placement mode and initial pass on new placement settings. Removal of concept of "selection" with new palette UX.

diff --git a/Engine/Source/Editor/AssetPlacementEdMode/Private/AssetPlacementEdModeToolkit.cpp b/Engine/Source/Editor/AssetPlacementEdMode/Private/AssetPlacementEdModeToolkit.cpp
--- a/Engine/Source/Editor/AssetPlacementEdMode/Private/AssetPlacementEdModeToolkit.cpp
+++ b/Engine/Source/Editor/AssetPlacementEdMode/Private/AssetPlacementEdModeToolkit.cpp
@@ -8,6 +8,7 @@
 #include "Modules/ModuleManager.h"
 
 #include "Widgets/Layout/SExpandableArea.h"
+#include "Widgets/Layout/SScrollBox.h"
 
 #include "EditorModeManager.h"
 
@@ -38,9 +39,17 @@ TSharedPtr<SWidget> FAssetPlacementEdModeToolkit::GetInlineContent() const
 {
 	return SNew(SVerticalBox)
 		+ SVerticalBox::Slot()
-		.AutoHeight()
+		.FillHeight(.6)
 		[
-			FModeToolkit::GetInlineContent().ToSharedRef()
+			SNew(SScrollBox)
+			+ SScrollBox::Slot()
+			[
+				DetailsView.ToSharedRef()
+			]
+			+ SScrollBox::Slot()
+			[
+				ModeDetailsView.ToSharedRef()
+			]
 		]
 		+ SVerticalBox::Slot()
 		.AutoHeight()
@@ -49,8 +58,12 @@ TSharedPtr<SWidget> FAssetPlacementEdModeToolkit::GetInlineContent() const
 			.AreaTitle(LOCTEXT("AssetPaletteHeader", "Asset Palette"))
 			.BodyContent()
 			[
-				SNew(SAssetPlacementPalette)
-				.PlacementSettings(PlacementSettings)
+				SNew(SScrollBox)
+				+ SScrollBox::Slot()
+				[
+					SNew(SAssetPlacementPalette)
+					.PlacementSettings(PlacementSettings)
+				]
 			]
 		];
 }
diff --git a/Engine/Source/Editor/AssetPlacementEdMode/Private/AssetPlacementPaletteItem.cpp b/Engine/Source/Editor/AssetPlacementEdMode/Private/AssetPlacementPaletteItem.cpp
--- a/Engine/Source/Editor/AssetPlacementEdMode/Private/AssetPlacementPaletteItem.cpp
+++ b/Engine/Source/Editor/AssetPlacementEdMode/Private/AssetPlacementPaletteItem.cpp
@@ -155,17 +155,6 @@ TSharedRef<SToolTip> FAssetPlacementPaletteItemModel::CreateTooltipWidget() cons
 		];
 }
 
-TSharedRef<SCheckBox> FAssetPlacementPaletteItemModel::CreateActivationCheckBox(TAttribute<bool> IsItemWidgetSelected, TAttribute<EVisibility> InVisibility)
-{
-	return
-		SNew(SCheckBox)
-		.Padding(0.f)
-		.OnCheckStateChanged(this, &FAssetPlacementPaletteItemModel::HandleCheckStateChanged, IsItemWidgetSelected)
-		.Visibility(InVisibility)
-		.IsChecked(this, &FAssetPlacementPaletteItemModel::GetCheckBoxState)
-		.ToolTipText(LOCTEXT("TileCheckboxTooltip", "Check to activate the currently selected types in the palette"));
-}
-
 FName FAssetPlacementPaletteItemModel::GetDisplayFName() const
 {
 	return DisplayFName;
@@ -183,16 +172,6 @@ FText FAssetPlacementPaletteItemModel::GetPaletteSearchText() const
 	}
 }
 
-void FAssetPlacementPaletteItemModel::SetTypeActiveInPalette(bool bSetActiveInPalette)
-{
-	bSelected = bSetActiveInPalette;
-}
-
-bool FAssetPlacementPaletteItemModel::IsActive() const
-{
-	return bSelected;
-}
-
 bool FAssetPlacementPaletteItemModel::IsBlueprint() const
 {
 	return TypeInfo && (TypeInfo->AssetData.GetClass()->GetDefaultObject()->IsA<UBlueprint>() || (TypeInfo->AssetData.GetClass()->ClassGeneratedBy != nullptr));
@@ -203,26 +182,6 @@ bool FAssetPlacementPaletteItemModel::IsAsset() const
 	return TypeInfo && TypeInfo->AssetData.IsValid();
 }
 
-void FAssetPlacementPaletteItemModel::HandleCheckStateChanged(const ECheckBoxState NewCheckedState, TAttribute<bool> IsItemWidgetSelected)
-{
-	if (!IsItemWidgetSelected.IsSet()) { return; }
-
-	const bool bShouldActivate = NewCheckedState == ECheckBoxState::Checked;
-	if (!IsItemWidgetSelected.Get())
-	{
-		SetTypeActiveInPalette(bShouldActivate);
-	}
-	else if (AssetPalette.IsValid())
-	{
-		AssetPalette.Pin()->ActivateAllSelectedTypes(bShouldActivate);
-	}
-}
-
-ECheckBoxState FAssetPlacementPaletteItemModel::GetCheckBoxState() const
-{
-	return bSelected ? ECheckBoxState::Checked : ECheckBoxState::Unchecked;
-}
-
 EVisibility FAssetPlacementPaletteItemModel::GetTooltipVisibility() const
 {
 	return (AssetPalette.IsValid() && AssetPalette.Pin()->ShouldShowTooltips()) ? EVisibility::SelfHitTestInvisible : EVisibility::Collapsed;
@@ -259,9 +218,6 @@ void SAssetPlacementPaletteItemTile::Construct(const FArguments& InArgs, TShared
 {
 	Model = InModel;
 
-	auto IsSelectedGetter = TAttribute<bool>::Create(TAttribute<bool>::FGetter::CreateSP(this, &SAssetPlacementPaletteItemTile::IsSelected));
-	auto CheckBoxVisibility = TAttribute<EVisibility>::Create(TAttribute<EVisibility>::FGetter::CreateSP(this, &SAssetPlacementPaletteItemTile::GetCheckBoxVisibility));
-
 	STableRow<FAssetPlacementUIInfoPtr>::Construct(
 		STableRow<FAssetPlacementUIInfoPtr>::FArguments()
 		.Style(FEditorStyle::Get(), "ContentBrowser.AssetListView.TableRow")
@@ -278,45 +234,14 @@ void SAssetPlacementPaletteItemTile::Construct(const FArguments& InArgs, TShared
 				.Padding(4.f)
 				.BorderImage(FEditorStyle::GetBrush("ContentBrowser.ThumbnailShadow"))
 				.ForegroundColor(FLinearColor::White)
-				.ColorAndOpacity(this, &SAssetPlacementPaletteItemTile::GetTileColorAndOpacity)
+				.ColorAndOpacity(FLinearColor::White)
 				[
 					Model->GetThumbnailWidget()
 				]
 			]
-			
-			// Checkbox
-			+ SOverlay::Slot()
-			.HAlign(HAlign_Left)
-			.VAlign(VAlign_Top)
-			.Padding(FMargin(3.f))
-			[
-				SNew(SBorder)
-				.BorderImage(FEditorStyle::GetBrush("ContentBrowser.ThumbnailShadow"))
-				.BorderBackgroundColor(FLinearColor::Black)
-				.ForegroundColor(FLinearColor::White)
-				.Padding(3.f)
-				[
-					Model->CreateActivationCheckBox(IsSelectedGetter, CheckBoxVisibility)
-				]
-			]
 		], InOwnerTableView);
 }
 
-FLinearColor SAssetPlacementPaletteItemTile::GetTileColorAndOpacity() const
-{
-	return Model->IsActive() ? FLinearColor::White : FLinearColor(0.5f, 0.5f, 0.5f, 1.f);
-}
-
-EVisibility SAssetPlacementPaletteItemTile::GetCheckBoxVisibility() const
-{
-	return CanShowOverlayItems() && (IsHovered() || (IsSelected() && Model->GetAssetPalette()->AnySelectedTileHovered())) ? EVisibility::Visible : EVisibility::Collapsed;
-}
-
-bool SAssetPlacementPaletteItemTile::CanShowOverlayItems() const
-{
-	return true;
-}
-
 ////////////////////////////////////////////////
 // SAssetPlacementPaletteItemRow
 ////////////////////////////////////////////////
@@ -333,12 +258,7 @@ TSharedRef<SWidget> SAssetPlacementPaletteItemRow::GenerateWidgetForColumn(const
 {
 	TSharedPtr<SWidget> TableRowContent = SNullWidget::NullWidget;
 
-	if (ColumnName == AssetPlacementPaletteTreeColumns::ColumnID_ToggleActive)
-	{
-		TAttribute<bool> IsSelectedGetter = MakeAttributeSP<bool>(this, &SAssetPlacementPaletteItemTile::IsSelected);
-		TableRowContent = Model->CreateActivationCheckBox(IsSelectedGetter);
-	} 
-	else if (ColumnName == AssetPlacementPaletteTreeColumns::ColumnID_Type)
+	if (ColumnName == AssetPlacementPaletteTreeColumns::ColumnID_Type)
 	{
 		TableRowContent =
 			SNew(SHorizontalBox)
diff --git a/Engine/Source/Editor/AssetPlacementEdMode/Private/AssetPlacementPaletteItem.h b/Engine/Source/Editor/AssetPlacementEdMode/Private/AssetPlacementPaletteItem.h
--- a/Engine/Source/Editor/AssetPlacementEdMode/Private/AssetPlacementPaletteItem.h
+++ b/Engine/Source/Editor/AssetPlacementEdMode/Private/AssetPlacementPaletteItem.h
@@ -25,15 +25,13 @@ typedef TSharedPtr<FPaletteItem> FAssetPlacementUIInfoPtr;
 
 namespace AssetPlacementPaletteTreeColumns
 {
-	/** IDs for list columns */
-	static const FName ColumnID_ToggleActive("Toggle");
 	static const FName ColumnID_Type("Type");
 };
 
 class FAssetPlacementPaletteItemModel : public TSharedFromThis<FAssetPlacementPaletteItemModel>
 {
 public:
-	FAssetPlacementPaletteItemModel(FAssetPlacementUIInfoPtr InTypeInfo, TSharedRef<class SAssetPlacementPalette> InFoliagePalette, TSharedPtr<class FAssetThumbnailPool> InThumbnailPool/*, FEdModeFoliage* InFoliageEditMode*/);
+	FAssetPlacementPaletteItemModel(FAssetPlacementUIInfoPtr InTypeInfo, TSharedRef<class SAssetPlacementPalette> InFoliagePalette, TSharedPtr<class FAssetThumbnailPool> InThumbnailPool);
 
 	/** @return The foliage palette that contains the item */
 	TSharedPtr<SAssetPlacementPalette> GetAssetPalette() const;
@@ -46,21 +44,12 @@ public:
 	/** @return The tooltip widget for this item */
 	TSharedRef<class SToolTip> CreateTooltipWidget() const;
 
-	/** @return The checkbox widget for activating/deactivating this type in the palette */
-	TSharedRef<class SCheckBox> CreateActivationCheckBox(TAttribute<bool> IsItemWidgetSelected, TAttribute<EVisibility> InVisibility = TAttribute<EVisibility>());
-
 	/** Gets the FName version of the displayed name of this item */
 	FName GetDisplayFName() const;
 
 	/** Gets the current search filter text */
 	FText GetPaletteSearchText() const;
 
-	/** Sets whether this foliage type is active in the palette */
-	void SetTypeActiveInPalette(bool bSetActiveInPalette);
-
-	/** Gets whether this foliage type is active in the palette */
-	bool IsActive() const;
-
 	/** @return Whether this palette item represents an instance of a foliage type blueprint class */
 	bool IsBlueprint() const;
 
@@ -68,12 +57,6 @@ public:
 	bool IsAsset() const;
 
 private:
-	/** Handles the change in activation of the item in the palette */
-	void HandleCheckStateChanged(const ECheckBoxState NewCheckedState, TAttribute<bool> IsItemWidgetSelected);
-
-	/** Gets whether the foliage type is active in the palette  */
-	ECheckBoxState GetCheckBoxState() const;
-
 	/** Gets the visibility of the entire tooltip */
 	EVisibility GetTooltipVisibility() const;
 
@@ -91,7 +74,6 @@ private:
 	FAssetPlacementUIInfoPtr TypeInfo;
 	TWeakPtr<SAssetPlacementPalette> AssetPalette;
 	TWeakPtr<IAssetTypeActions> AssetTypeActions;
-	bool bSelected = true;
 };
 
 /** A tile representing a foliage type in the palette */
@@ -103,11 +85,6 @@ public:
 
 	void Construct(const FArguments& InArgs, TSharedRef<STableViewBase> InOwnerTableView, TSharedPtr<FAssetPlacementPaletteItemModel>& InModel);
 
-private:
-	FLinearColor GetTileColorAndOpacity() const;
-	EVisibility GetCheckBoxVisibility() const;
-	bool CanShowOverlayItems() const;
-
 private:
 	static const float MinScaleForOverlayItems;
 	TSharedPtr<FAssetPlacementPaletteItemModel> Model;
diff --git a/Engine/Source/Editor/AssetPlacementEdMode/Private/AssetPlacementSettings.cpp b/Engine/Source/Editor/AssetPlacementEdMode/Private/AssetPlacementSettings.cpp
new file mode 100644
--- /dev/null
+++ b/Engine/Source/Editor/AssetPlacementEdMode/Private/AssetPlacementSettings.cpp
@@ -0,0 +1,43 @@
+// Copyright Epic Games, Inc. All Rights Reserved.
+
+#include "AssetPlacementSettings.h"
+
+bool UAssetPlacementSettings::CanEditChange(const FProperty* InProperty) const
+{
+	if (!Super::CanEditChange(InProperty))
+	{
+		return false;
+	}
+
+	const FName PropertyName = InProperty->GetFName();
+	if (PropertyName == GET_MEMBER_NAME_CHECKED(UAssetPlacementSettings, ScaleRangeUniform) || PropertyName == GET_MEMBER_NAME_CHECKED(UAssetPlacementSettings, bAllowNegativeUniformScale))
+	{
+		return (ScalingType == EFoliageScaling::Uniform) && bUseRandomScale;
+	}
+
+	const bool bCanFreeScale = (ScalingType != EFoliageScaling::Uniform) && bUseRandomScale;
+	if (PropertyName == GET_MEMBER_NAME_CHECKED(UAssetPlacementSettings, ScaleRangeX) || PropertyName == GET_MEMBER_NAME_CHECKED(UAssetPlacementSettings, bAllowNegativeXScale))
+	{
+		const bool bLockX = (ScalingType == EFoliageScaling::LockXY) || (ScalingType == EFoliageScaling::LockXZ);
+		return bCanFreeScale && !bLockX;
+	}
+
+	if (PropertyName == GET_MEMBER_NAME_CHECKED(UAssetPlacementSettings, ScaleRangeY) || PropertyName == GET_MEMBER_NAME_CHECKED(UAssetPlacementSettings, bAllowNegativeYScale))
+	{
+		const bool bLockY = (ScalingType == EFoliageScaling::LockXY) || (ScalingType == EFoliageScaling::LockYZ);
+		return bCanFreeScale && !bLockY;
+	}
+
+	if (PropertyName == GET_MEMBER_NAME_CHECKED(UAssetPlacementSettings, FreeScaleRangeZ) || PropertyName == GET_MEMBER_NAME_CHECKED(UAssetPlacementSettings, bAllowNegativeZScale))
+	{
+		const bool bLockZ = (ScalingType == EFoliageScaling::LockYZ) || (ScalingType == EFoliageScaling::LockXZ);
+		return bCanFreeScale && !bLockZ;
+	}
+
+	if (PropertyName == GET_MEMBER_NAME_CHECKED(UAssetPlacementSettings, bInvertNormalAxis))
+	{
+		return bAlignToNormal;
+	}
+
+	return true;
+}
diff --git a/Engine/Source/Editor/AssetPlacementEdMode/Private/AssetPlacementSettings.h b/Engine/Source/Editor/AssetPlacementEdMode/Private/AssetPlacementSettings.h
--- a/Engine/Source/Editor/AssetPlacementEdMode/Private/AssetPlacementSettings.h
+++ b/Engine/Source/Editor/AssetPlacementEdMode/Private/AssetPlacementSettings.h
@@ -4,6 +4,7 @@
 
 #include "CoreMinimal.h"
 #include "AssetRegistry/AssetData.h"
+#include "FoliageType.h"
 
 #include "AssetPlacementSettings.generated.h"
 
@@ -19,9 +20,6 @@ struct FPaletteItem
 
 	UPROPERTY()
 	TScriptInterface<IAssetFactoryInterface> FactoryOverride;
-
-	UPROPERTY()
-	bool bIsEnabled = true;
 };
 
 UCLASS(config = EditorPerProjectUserSettings)
@@ -30,32 +28,82 @@ class UAssetPlacementSettings : public UObject
 	GENERATED_BODY()
 
 public:
-	UPROPERTY(config, EditAnywhere, Category = "Filters")
+	UPROPERTY(config, EditAnywhere, Category = "Axis Alignment", meta = (InlineEditConditionToggle))
+	bool bAlignToNormal = true;
+	
+	UPROPERTY(config, EditAnywhere, Category = "Axis Alignment", meta = (DisplayName = "Align To Normal", EditCondition = "bAlignToNormal"))
+	TEnumAsByte<EAxis::Type> AxisToAlignWithNormal = EAxis::Type::Z;
+
+	UPROPERTY(config, EditAnywhere, Category = "Axis Alignment")
+	bool bInvertNormalAxis = false;
+	
+	UPROPERTY(config, EditAnywhere, Category = "Rotation", meta = (InlineEditConditionToggle))
+	bool bUseRandomXRotation = false;
+
+	UPROPERTY(config, EditAnywhere, AdvancedDisplay, Category = "Rotation", meta = (EditCondition = "bUseRandomXRotation", UIMin = "0.0", UIMax = "360.0"))
+	FFloatInterval RandomRotationX = FFloatInterval(0.0f, 360.0f);
+	
+	UPROPERTY(config, EditAnywhere, Category = "Rotation", meta = (InlineEditConditionToggle))
+	bool bUseRandomYRotation = false;
+
+	UPROPERTY(config, EditAnywhere, AdvancedDisplay, Category = "Rotation", meta = (EditCondition = "bUseRandomYRotation", UIMin = "0.0", UIMax = "360.0"))
+	FFloatInterval RandomRotationY = FFloatInterval(0.0f, 360.0f);
+
+	UPROPERTY(config, EditAnywhere, Category = "Rotation", meta = (InlineEditConditionToggle))
+	bool bUseRandomZRotation = true;
+
+	UPROPERTY(config, EditAnywhere, Category = "Rotation", meta = (EditCondition = "bUseRandomZRotation", UIMin = "0.0", UIMax = "360.0"))
+	FFloatInterval RandomRotationZ = FFloatInterval(0.0f, 360.0f);
+
+	UPROPERTY(config, EditAnywhere, Category = "Scale", meta = (InlineEditConditionToggle))
+	bool bUseRandomScale = true;
+
+	UPROPERTY(config, EditAnywhere, Category = "Scale", meta = (DisplayName="Use Random Scale", EditCondition = "bAllowRandomScale"))
+	EFoliageScaling ScalingType = EFoliageScaling::Uniform;
+	
+	UPROPERTY(config, EditAnywhere, AdvancedDisplay, Category = "Scale", meta = (UIMin = ".01"))
+	FFloatInterval ScaleRangeUniform = FFloatInterval(.8f, 1.0f);
+	
+	UPROPERTY(config, EditAnywhere, AdvancedDisplay, Category = "Scale")
+	bool bAllowNegativeUniformScale = false;
+	
+	UPROPERTY(config, EditAnywhere, AdvancedDisplay, Category = "Scale", meta = (UIMin = ".01"))
+	FFloatInterval ScaleRangeX = FFloatInterval(.8f, 1.0f);
+	
+	UPROPERTY(config, EditAnywhere, AdvancedDisplay, Category = "Scale")
+	bool bAllowNegativeXScale = false;
+
+	UPROPERTY(config, EditAnywhere, AdvancedDisplay, Category = "Scale", meta = (UIMin = ".01"))
+	FFloatInterval ScaleRangeY = FFloatInterval(.8f, 1.0f);
+	
+	UPROPERTY(config, EditAnywhere, AdvancedDisplay, Category = "Scale")
+	bool bAllowNegativeYScale = false;
+
+	UPROPERTY(config, EditAnywhere, AdvancedDisplay, Category = "Scale", meta = (UIMin = ".01"))
+	FFloatInterval FreeScaleRangeZ = FFloatInterval(.8f, 1.0f);
+	
+	UPROPERTY(config, EditAnywhere, AdvancedDisplay, Category = "Scale")
+	bool bAllowNegativeZScale = false;
+
+	UPROPERTY(config, EditAnywhere, Category = "Place on Filters")
 	bool bLandscape = true;
 
-	UPROPERTY(config, EditAnywhere, Category = "Filters")
+	UPROPERTY(config, EditAnywhere, Category = "Place on Filters")
 	bool bStaticMeshes = true;
 
-	UPROPERTY(config, EditAnywhere, Category = "Filters")
+	UPROPERTY(config, EditAnywhere, Category = "Place on Filters")
 	bool bBSP = true;
 
-	UPROPERTY(config, EditAnywhere, Category = "Filters")
+	UPROPERTY(config, EditAnywhere, Category = "Place on Filters")
 	bool bFoliage = false;
 
-	UPROPERTY(config, EditAnywhere, Category = "Filters")
+	UPROPERTY(config, EditAnywhere, Category = "Place on Filters")
 	bool bTranslucent = false;
 
-	UPROPERTY(config, EditAnywhere, Category = "Placement")
-	bool bAllowRandomRotation = true;
-
-	UPROPERTY(config, EditAnywhere, Category = "Placement")
-	bool bAllowAlignToNormal = true;
-
-	UPROPERTY(config, EditAnywhere, Category = "Placement")
-	bool bAllowRandomScale = true;
-
 	// todo: asset data does not serialize out correctly
 	// maybe save soft object pointers, and convert in the UI to asset data?
 	//UPROPERTY(config)
 	TArray<FPaletteItem> PaletteItems;
+
+	virtual bool CanEditChange(const FProperty* InProperty) const override;
 };
diff --git a/Engine/Source/Editor/AssetPlacementEdMode/Private/SAssetPlacementPalette.cpp b/Engine/Source/Editor/AssetPlacementEdMode/Private/SAssetPlacementPalette.cpp
--- a/Engine/Source/Editor/AssetPlacementEdMode/Private/SAssetPlacementPalette.cpp
+++ b/Engine/Source/Editor/AssetPlacementEdMode/Private/SAssetPlacementPalette.cpp
@@ -27,7 +27,6 @@
 #include "AssetSelection.h"
 #include "ScopedTransaction.h"
 #include "AssetData.h"
-#include "Engine/StaticMesh.h"
 #include "Editor.h"
 #include "Widgets/Layout/SScrollBox.h"
 
@@ -36,8 +35,9 @@
 #include "Widgets/Input/SSearchBox.h"
 #include "AssetPlacementEdModeStyle.h"
 #include "AssetPlacementSettings.h"
-#include "Selection.h"
 #include "Subsystems/PlacementSubsystem.h"
+#include "Engine/AssetManager.h"
+#include "Engine/StreamableManager.h"
 
 #define LOCTEXT_NAMESPACE "AssetPlacementMode"
 
@@ -103,60 +103,6 @@ private:
 	bool bIsDragOn;
 };
 
-////////////////////////////////////////////////
-// SUneditablePlacementTypeWarning
-////////////////////////////////////////////////
-class SUneditableAssetTypeWarning : public SCompoundWidget
-{
-public:
-	SLATE_BEGIN_ARGS(SUneditableAssetTypeWarning)
-		: _WarningText()
-		, _OnHyperlinkClicked()
-	{}
-
-	/** The rich text to show in the warning */
-	SLATE_ATTRIBUTE(FText, WarningText)
-
-		/** Called when the hyperlink in the rich text is clicked */
-		SLATE_EVENT(FSlateHyperlinkRun::FOnClick, OnHyperlinkClicked)
-
-		SLATE_END_ARGS()
-
-		/** Constructs the widget */
-		void Construct(const FArguments& InArgs)
-	{
-		ChildSlot
-			[
-				SNew(SBorder)
-				.BorderImage(FAppStyle::Get().GetBrush("ToolPanel.GroupBorder"))
-				[
-					SNew(SHorizontalBox)
-					+ SHorizontalBox::Slot()
-					.AutoWidth()
-					.HAlign(HAlign_Center)
-					.VAlign(VAlign_Center)
-					.Padding(2)
-					[
-						SNew(SImage)
-						.Image(FAppStyle::Get().GetBrush("Icons.Warning"))
-					]
-					+ SHorizontalBox::Slot()
-						.VAlign(VAlign_Center)
-						.Padding(2)
-						[
-							SNew(SRichTextBlock)
-							.DecoratorStyleSet(&FAppStyle::Get())
-							.Justification(ETextJustify::Left)
-							.TextStyle(FAppStyle::Get(), "DetailsView.BPMessageTextStyle")
-							.Text(InArgs._WarningText)
-							.AutoWrapText(true)
-							+ SRichTextBlock::HyperlinkDecorator(TEXT("HyperlinkDecorator"), InArgs._OnHyperlinkClicked)
-						]
-				]
-			];
-	}
-};
-
 ////////////////////////////////////////////////
 // SPlacementPalette
 ////////////////////////////////////////////////
@@ -319,29 +265,6 @@ void SAssetPlacementPalette::RefreshPalette()
 	}
 }
 
-bool SAssetPlacementPalette::AnySelectedTileHovered() const
-{
-	for (FPlacementPaletteItemModelPtr& PaletteItem : GetActiveViewWidget()->GetSelectedItems())
-	{
-		TSharedPtr<ITableRow> Tile = TileViewWidget->WidgetFromItem(PaletteItem);
-		if (Tile.IsValid() && Tile->AsWidget()->IsHovered())
-		{
-			return true;
-		}
-	}
-
-	return false;
-}
-
-void SAssetPlacementPalette::ActivateAllSelectedTypes(bool bActivate) const
-{
-	// Apply the new check state to all of the selected types
-	for (FPlacementPaletteItemModelPtr& PaletteItem : GetActiveViewWidget()->GetSelectedItems())
-	{
-		PaletteItem->SetTypeActiveInPalette(bActivate);
-	}
-}
-
 void SAssetPlacementPalette::RefreshActivePaletteViewWidget()
 {
 	if (ActiveViewMode == EAssetPlacementPaletteViewMode::Thumbnail)
@@ -382,12 +305,15 @@ void SAssetPlacementPalette::AddPlacementType(const FAssetData& AssetData)
 		return;
 	}
 
-	// todo - can the palette items be a map? will this break the list view?
+	// Check that we don't already have this item in the palette
 	if (PaletteItems.FindByPredicate([&AssetData](TSharedPtr<FAssetPlacementPaletteItemModel>& InItem) { return InItem.IsValid() && (InItem->GetTypeUIInfo()->AssetData.ObjectPath == AssetData.ObjectPath); }))
 	{
 		return;
 	}
 
+	// Try to load the asset async so it's ready to place.
+	UAssetManager::GetStreamableManager().RequestAsyncLoad(AssetData.ToSoftObjectPath());
+
 	FAssetPlacementUIInfoPtr PlacementInfo = MakeShared<FPaletteItem>();
 	PlacementInfo->AssetData = AssetData;
 	PlacementInfo->FactoryOverride = FactoryInterface;
@@ -406,41 +332,21 @@ TSharedRef<SWidgetSwitcher> SAssetPlacementPalette::CreatePaletteViews()
 	// Tile View Widget
 	SAssignNew(TileViewWidget, SPlacementTypeTileView)
 		.ListItemsSource(&FilteredItems)
-		.SelectionMode(ESelectionMode::Multi)
 		.OnGenerateTile(this, &SAssetPlacementPalette::GenerateTile)
 		.OnContextMenuOpening(this, &SAssetPlacementPalette::ConstructPlacementTypeContextMenu)
-		.OnSelectionChanged(this, &SAssetPlacementPalette::OnSelectionChanged)
 		.ItemHeight(this, &SAssetPlacementPalette::GetScaledThumbnailSize)
 		.ItemWidth(this, &SAssetPlacementPalette::GetScaledThumbnailSize)
-		.ItemAlignment(EListItemAlignment::LeftAligned)
-		.ClearSelectionOnClick(true)
-		.OnMouseButtonDoubleClick(this, &SAssetPlacementPalette::OnItemDoubleClicked);
+		.ItemAlignment(EListItemAlignment::LeftAligned);
 
 	// Tree View Widget
 	SAssignNew(TreeViewWidget, SPlacementTypeTreeView)
 		.TreeItemsSource(&FilteredItems)
-		.SelectionMode(ESelectionMode::Multi)
 		.OnGenerateRow(this, &SAssetPlacementPalette::TreeViewGenerateRow)
 		.OnGetChildren(this, &SAssetPlacementPalette::TreeViewGetChildren)
 		.OnContextMenuOpening(this, &SAssetPlacementPalette::ConstructPlacementTypeContextMenu)
-		.OnSelectionChanged(this, &SAssetPlacementPalette::OnSelectionChanged)
-		.OnMouseButtonDoubleClick(this, &SAssetPlacementPalette::OnItemDoubleClicked)
 		.HeaderRow
 		(
-			// Toggle Active
 			SAssignNew(TreeViewHeaderRow, SHeaderRow)
-			+ SHeaderRow::Column(AssetPlacementPaletteTreeColumns::ColumnID_ToggleActive)
-			[
-				SNew(SCheckBox)
-				.IsChecked(this, &SAssetPlacementPalette::GetState_AllMeshes)
-			.OnCheckStateChanged(this, &SAssetPlacementPalette::OnCheckStateChanged_AllMeshes)
-			]
-			.DefaultLabel(BlankText)
-			.HeaderContentPadding(FMargin(0, 1, 0, 1))
-			.HAlignHeader(HAlign_Center)
-			.HAlignCell(HAlign_Center)
-			.FixedWidth(24)
-
 			// Type
 			+ SHeaderRow::Column(AssetPlacementPaletteTreeColumns::ColumnID_Type)
 			.HeaderContentPadding(FMargin(10, 1, 0, 1))
@@ -524,27 +430,6 @@ void SAssetPlacementPalette::SetViewMode(EAssetPlacementPaletteViewMode NewViewM
 	if (ActiveViewMode != NewViewMode)
 	{
 		ActiveViewMode = NewViewMode;
-		switch (ActiveViewMode)
-		{
-		case EAssetPlacementPaletteViewMode::Thumbnail:
-			// Set the tile selection to be the current tree selections
-			TileViewWidget->ClearSelection();
-			for (auto& TypeInfo : TreeViewWidget->GetSelectedItems())
-			{
-				TileViewWidget->SetItemSelection(TypeInfo, true);
-			}
-			break;
-			
-		case EAssetPlacementPaletteViewMode::Tree:
-			// Set the tree selection to be the current tile selection
-			TreeViewWidget->ClearSelection();
-			for (auto& TypeInfo : TileViewWidget->GetSelectedItems())
-			{
-				TreeViewWidget->SetItemSelection(TypeInfo, true);
-			}
-			break;
-		}
-
 		WidgetSwitcher->SetActiveWidgetIndex((int32)ActiveViewMode);
 		
 		RefreshActivePaletteViewWidget();
@@ -571,16 +456,6 @@ FText SAssetPlacementPalette::GetSearchText() const
 	return TypeFilter->GetRawFilterText();
 }
 
-void SAssetPlacementPalette::OnSelectionChanged(FPlacementPaletteItemModelPtr Item, ESelectInfo::Type SelectInfo)
-{
-	// Not yet implemented
-}
-
-void SAssetPlacementPalette::OnItemDoubleClicked(FPlacementPaletteItemModelPtr Item) const
-{
-	Item->SetTypeActiveInPalette(!Item->IsActive());
-}
-
 TSharedRef<SWidget> SAssetPlacementPalette::GetViewOptionsMenuContent()
 {
 	FMenuBuilder MenuBuilder(true, UICommandList);
@@ -720,175 +595,21 @@ TSharedPtr<SWidget> SAssetPlacementPalette::ConstructPlacementTypeContextMenu()
 	return MenuBuilder.MakeWidget();
 }
 
-void SAssetPlacementPalette::OnActivatePlacementTypes()
-{
-	for (FPlacementPaletteItemModelPtr& PaletteItem : GetActiveViewWidget()->GetSelectedItems())
-	{
-		PaletteItem->SetTypeActiveInPalette(true);
-	}
-}
-
-bool SAssetPlacementPalette::OnCanActivatePlacementTypes() const
-{
-	// At least one selected item must be inactive
-	for (FPlacementPaletteItemModelPtr& PaletteItem : GetActiveViewWidget()->GetSelectedItems())
-	{
-		if (!PaletteItem->IsActive())
-		{
-			return true;
-		}
-	}
-
-	return false;
-}
-
-void SAssetPlacementPalette::OnDeactivatePlacementTypes()
-{
-	for (FPlacementPaletteItemModelPtr& PaletteItem : GetActiveViewWidget()->GetSelectedItems())
-	{
-		PaletteItem->SetTypeActiveInPalette(false);
-	}
-}
-
-bool SAssetPlacementPalette::OnCanDeactivatePlacementTypes() const
-{
-	// At least one selected item must be active
-	for (FPlacementPaletteItemModelPtr& PaletteItem : GetActiveViewWidget()->GetSelectedItems())
-	{
-		if (PaletteItem->IsActive())
-		{
-			return true;
-		}
-	}
-
-	return false;
-}
-
-void SAssetPlacementPalette::FillReplacePlacementTypeSubmenu(FMenuBuilder& MenuBuilder)
-{
-	FContentBrowserModule& ContentBrowserModule = FModuleManager::Get().LoadModuleChecked<FContentBrowserModule>(TEXT("ContentBrowser"));
-
-	FAssetPickerConfig AssetPickerConfig;
-	AssetPickerConfig.Filter.bRecursiveClasses = true;
-	AssetPickerConfig.OnAssetSelected = FOnAssetSelected::CreateSP(this, &SAssetPlacementPalette::OnReplacePlacementTypeSelected);
-	AssetPickerConfig.InitialAssetViewType = EAssetViewType::List;
-	AssetPickerConfig.bAllowNullSelection = false;
-
-	TSharedRef<SWidget> MenuContent = SNew(SBox)
-		.WidthOverride(384)
-		.HeightOverride(500)
-		[
-			ContentBrowserModule.Get().CreateAssetPicker(AssetPickerConfig)
-		];
-
-	MenuBuilder.AddWidget(MenuContent, FText::GetEmpty(), true);
-}
-
-void SAssetPlacementPalette::OnReplacePlacementTypeSelected(const FAssetData& AssetData)
-{
-	FSlateApplication::Get().DismissAllMenus();
-
-	// not yet implemented
-	// should replace existing matching elements in level with the newly selected asset data
-	// and update the palette
-}
-
-void SAssetPlacementPalette::OnRemovePlacementType()
-{
-	TArray<FPlacementPaletteItemModelPtr> PaletteCopy = PaletteItems;
-	for (FPlacementPaletteItemModelPtr& PaletteItem : GetActiveViewWidget()->GetSelectedItems())
-	{
-		PaletteCopy.Remove(PaletteItem);
-	}
-
-	PaletteItems = MoveTemp(PaletteCopy);
-	GetActiveViewWidget()->RequestListRefresh();
-}
-
 void SAssetPlacementPalette::OnShowPlacementTypeInCB()
 {
-	TArray<FAssetData> SelectedAssets;
-	for (FPlacementPaletteItemModelPtr& PaletteItem : GetActiveViewWidget()->GetSelectedItems())
-	{
-		if (PaletteItem.IsValid() && PaletteItem->GetTypeUIInfo().IsValid())
-		{
-			SelectedAssets.Add(PaletteItem->GetTypeUIInfo()->AssetData);
-		}
-	}
-
-	if (SelectedAssets.Num())
-	{
-		GEditor->SyncBrowserToObjects(SelectedAssets);
-	}
-}
-
-void SAssetPlacementPalette::OnReflectSelectionInPalette()
-{
-	TArray<FAssetData> SelectedPlacementTypes;
-	for (UObject* SelectedObject : GEditor->GetSelectedObjects()->GetElementSelectionSet()->GetSelectedObjects())
-	{
-		SelectedPlacementTypes.Add(FAssetData(SelectedObject));
-	}
-	SelectPlacementTypesInPalette(SelectedPlacementTypes);
-}
-
-void SAssetPlacementPalette::SelectPlacementTypesInPalette(const TArray<FAssetData>& PlacementTypes)
-{
-	TArray<FPlacementPaletteItemModelPtr> SelectedItems;
-	SelectedItems.Reserve(PlacementTypes.Num());
-
+	TArray<FAssetData> FilteredAssets;
 	for (FPlacementPaletteItemModelPtr& PaletteItem : FilteredItems)
 	{
-		if (PlacementTypes.Contains(PaletteItem->GetTypeUIInfo()->AssetData))
+		if (PaletteItem.IsValid() && PaletteItem->GetTypeUIInfo().IsValid())
 		{
-			SelectedItems.Add(PaletteItem);
+			FilteredAssets.Add(PaletteItem->GetTypeUIInfo()->AssetData);
 		}
 	}
-	
-	GetActiveViewWidget()->ClearSelection();
-	GetActiveViewWidget()->SetItemSelection(SelectedItems, true);
-}
 
-void SAssetPlacementPalette::ExecuteOnSelectedItemPlacementTypes(TFunctionRef<void(const TArray<FAssetData>&)> ExecuteFunc)
-{
-	TArray<FPlacementPaletteItemModelPtr> SelectedItems = GetActiveViewWidget()->GetSelectedItems();
-	TArray<FAssetData> PlacementTypes;
-	PlacementTypes.Reserve(SelectedItems.Num());
-	for (FPlacementPaletteItemModelPtr& PaletteItem : SelectedItems)
+	if (FilteredAssets.Num())
 	{
-		PlacementTypes.Add(PaletteItem->GetTypeUIInfo()->AssetData);
+		GEditor->SyncBrowserToObjects(FilteredAssets);
 	}
-	ExecuteFunc(PlacementTypes);
-}
-
-void SAssetPlacementPalette::OnSelectAllInstances()
-{
-	ExecuteOnSelectedItemPlacementTypes([&](const TArray<FAssetData>& PlacementTypes)
-	{
-		//PlacementEditMode->SelectInstances(PlacementTypes, true);
-	});	
-}
-
-void SAssetPlacementPalette::OnDeselectAllInstances()
-{
-	ExecuteOnSelectedItemPlacementTypes([&](const TArray<FAssetData>& PlacementTypes)
-	{
-		//PlacementEditMode->SelectInstances(PlacementTypes, false);
-	});
-}
-
-void SAssetPlacementPalette::OnSelectInvalidInstances()
-{
-	ExecuteOnSelectedItemPlacementTypes([&](const TArray<FAssetData>& PlacementTypes)
-	{
-		//PlacementEditMode->SelectInvalidInstances(PlacementTypes);
-	});
-}
-
-bool SAssetPlacementPalette::CanSelectInstances() const
-{
-	return false;
-	//return PlacementEditMode->UISettings.GetSelectToolSelected() || PlacementEditMode->UISettings.GetLassoSelectToolSelected();
 }
 
 // THUMBNAIL VIEW
@@ -904,7 +625,7 @@ TSharedRef<ITableRow> SAssetPlacementPalette::GenerateTile(FPlacementPaletteItem
 float SAssetPlacementPalette::GetScaledThumbnailSize() const
 {
 	const FInt32Interval& SizeRange = PlacementPaletteConstants::ThumbnailSizeRange;
-	return SizeRange.Min + SizeRange.Size() * GetThumbnailScale();// PlacementEditMode->UISettings.GetPaletteThumbnailScale();
+	return SizeRange.Min + SizeRange.Size() * GetThumbnailScale();
 }
 
 float SAssetPlacementPalette::GetThumbnailScale() const
@@ -934,40 +655,6 @@ void SAssetPlacementPalette::TreeViewGetChildren(FPlacementPaletteItemModelPtr I
 	// Items do not have any children
 }
 
-ECheckBoxState SAssetPlacementPalette::GetState_AllMeshes() const
-{
-	bool bHasChecked = false;
-	bool bHasUnchecked = false;
-
-	for (const FPlacementPaletteItemModelPtr& PaletteItem : FilteredItems)
-	{
-		if (PaletteItem->IsActive())
-		{
-			bHasChecked = true;
-		}
-		else
-		{
-			bHasUnchecked = true;
-		}
-
-		if (bHasChecked && bHasUnchecked)
-		{
-			return ECheckBoxState::Undetermined;
-		}
-	}
-
-	return bHasChecked ? ECheckBoxState::Checked : ECheckBoxState::Unchecked;
-}
-
-void SAssetPlacementPalette::OnCheckStateChanged_AllMeshes(ECheckBoxState InState)
-{
-	const bool bActivate = InState == ECheckBoxState::Checked;
-	for (FPlacementPaletteItemModelPtr& PaletteItem : FilteredItems)
-	{
-		PaletteItem->SetTypeActiveInPalette(bActivate);
-	}
-}
-
 FText SAssetPlacementPalette::GetTypeColumnHeaderText() const
 {
 	return LOCTEXT("PlacementTypeHeader", "Asset Type");
@@ -999,51 +686,11 @@ void SAssetPlacementPalette::OnTypeColumnSortModeChanged(EColumnSortPriority::Ty
 	}
 }
 
-void SAssetPlacementPalette::OnEditPlacementTypeBlueprintHyperlinkClicked(const FSlateHyperlinkRun::FMetadata& Metadata)
-{
-	UBlueprint* Blueprint = nullptr; 
-	
-	// Get the first selected Placement type blueprint
-	for (FPlacementPaletteItemModelPtr& PaletteItem : GetActiveViewWidget()->GetSelectedItems())
-	{
-		Blueprint = Cast<UBlueprint>(PaletteItem->GetTypeUIInfo()->AssetData.GetClass()->ClassGeneratedBy);
-		if (Blueprint != nullptr)
-		{
-			break;
-		}
-	}
-
-	if (Blueprint)
-	{
-		// Open the blueprint
-		GEditor->EditObject(Blueprint);
-	}
-}
-
 EActiveTimerReturnType SAssetPlacementPalette::UpdatePaletteItems(double InCurrentTime, float InDeltaTime)
 {
 	if (bItemsNeedRebuild)
 	{
 		bItemsNeedRebuild = false;
-
-		// Cache the currently selected items
-		auto ActiveViewWidget = GetActiveViewWidget();
-		TArray<FPlacementPaletteItemModelPtr> PreviouslySelectedItems = ActiveViewWidget->GetSelectedItems();
-		ActiveViewWidget->ClearSelection();
-
-		// Restore the selection
-		for (auto& PrevSelectedItem : PreviouslySelectedItems)
-		{
-			// Select any replacements for previously selected Placement types
-			for (auto& Item : PaletteItems)
-			{
-				if (Item->GetDisplayFName() == PrevSelectedItem->GetDisplayFName())
-				{
-					ActiveViewWidget->SetItemSelection(Item, true);
-					break;
-				}
-			}
-		}
 	}
 
 	// Update the filtered items
@@ -1057,7 +704,6 @@ EActiveTimerReturnType SAssetPlacementPalette::UpdatePaletteItems(double InCurre
 
 		FPaletteItem NewItem;
 		NewItem.AssetData = Item->GetTypeUIInfo()->AssetData;
-		NewItem.bIsEnabled = false;
 	}
 
 	// Refresh the appropriate view
diff --git a/Engine/Source/Editor/AssetPlacementEdMode/Private/SAssetPlacementPalette.h b/Engine/Source/Editor/AssetPlacementEdMode/Private/SAssetPlacementPalette.h
--- a/Engine/Source/Editor/AssetPlacementEdMode/Private/SAssetPlacementPalette.h
+++ b/Engine/Source/Editor/AssetPlacementEdMode/Private/SAssetPlacementPalette.h
@@ -59,9 +59,6 @@ public:
 	/** Refreshes the Placement palette */
 	void RefreshPalette();
 
-	bool AnySelectedTileHovered() const;
-	void ActivateAllSelectedTypes(bool bActivate) const;
-
 	/** @return True if the given view mode is the active view mode */
 	bool IsActiveViewMode(EAssetPlacementPaletteViewMode ViewMode) const;
 
@@ -96,10 +93,6 @@ private:	// GENERAL
 	/** Gets the visibility of the Add Placement Type text in the header row button */
 	EVisibility GetAddPlacementTypeButtonTextVisibility() const;
 
-	/** Toggle whether all Placement types are active */
-	ECheckBoxState GetState_AllMeshes() const;
-	void OnCheckStateChanged_AllMeshes(ECheckBoxState InState);
-
 	/** Sets the view mode of the palette */
 	void SetViewMode(EAssetPlacementPaletteViewMode NewViewMode);
 
@@ -112,12 +105,6 @@ private:	// GENERAL
 	/** @return The index of the view widget to display */
 	int32 GetActiveViewIndex() const;
 
-	/** Handler for selection changes in either view */
-	void OnSelectionChanged(FPlacementPaletteItemModelPtr Item, ESelectInfo::Type SelectInfo);
-
-	/** Toggle the activation state of a type on a double-click */
-	void OnItemDoubleClicked(FPlacementPaletteItemModelPtr Item) const;
-
 	/** Creates the view options menu */
 	TSharedRef<SWidget> GetViewOptionsMenuContent();
 
@@ -137,47 +124,9 @@ private:	// CONTEXT MENU
 	/** @return the SWidget containing the context menu */
 	TSharedPtr<SWidget> ConstructPlacementTypeContextMenu();
 
-	/** Handler for the 'Activate' command */
-	void OnActivatePlacementTypes();
-	bool OnCanActivatePlacementTypes() const;
-
-	/** Handler for the 'Deactivate' command */
-	void OnDeactivatePlacementTypes();
-	bool OnCanDeactivatePlacementTypes() const;
-
-	/** Fills 'Replace' menu command  */
-	void FillReplacePlacementTypeSubmenu(FMenuBuilder& MenuBuilder);
-
-	/** Handler for 'Replace' command  */
-	void OnReplacePlacementTypeSelected(const struct FAssetData& AssetData);
-
-	/** Handler for 'Remove' command  */
-	void OnRemovePlacementType();
-
 	/** Handler for 'Show in CB' command  */
 	void OnShowPlacementTypeInCB();
 
-	/** Handler for 'Select All' command  */
-	void OnSelectAllInstances();
-
-	/** Handler for 'Deselect All' command  */
-	void OnDeselectAllInstances();
-
-	/** Handler for 'Select Invalid Instances' command  */
-	void OnSelectInvalidInstances();
-
-	/** Executes Function on gathered list of Placement types from the selected palette items */
-	void ExecuteOnSelectedItemPlacementTypes(TFunctionRef<void(const TArray<FAssetData>&)> ExecuteFunc);
-
-	/** @return Whether selecting instances is currently possible */
-	bool CanSelectInstances() const;
-
-	/** Handler for 'Reflect Selection in Palette ' command */
-	void OnReflectSelectionInPalette();
-
-	/** Selects Placement Type in palette */
-	void SelectPlacementTypesInPalette(const TArray<FAssetData>& PlacementTypes);
-
 private:	// THUMBNAIL VIEW
 
 	/** Creates a thumbnail tile for the given Placement type */
diff --git a/Engine/Source/Editor/AssetPlacementEdMode/Private/Tools/PlacementBrushToolBase.cpp b/Engine/Source/Editor/AssetPlacementEdMode/Private/Tools/PlacementBrushToolBase.cpp
--- a/Engine/Source/Editor/AssetPlacementEdMode/Private/Tools/PlacementBrushToolBase.cpp
+++ b/Engine/Source/Editor/AssetPlacementEdMode/Private/Tools/PlacementBrushToolBase.cpp
@@ -92,21 +92,12 @@ FTransform UPlacementBrushToolBase::GetFinalTransformFromHitLocationAndNormal(co
 		return FinalizedTransform;
 	}
 
-	// For now, just apply a random yaw to any placed object, until we have per object settings for random pitch and yaw angles.
-	if (PlacementSettings->bAllowRandomRotation)
-	{
-		FRotator UpdatedRotation = FinalizedTransform.Rotator();
-
-		// UpdatedRotation = FRotator(FMath::FRand() * ItemToPlace->RandomPitchAngle, 0.f, 0.f);
-
-		// if (ItemToPlace.bUseRandomYaw)
-		UpdatedRotation.Yaw = FMath::FRand() * 360.f;
-
-		FinalizedTransform.SetRotation(UpdatedRotation.Quaternion());
-	}
+	// Update the rotation if we need to
+	FRotator FinalRotation = GetFinalRotation(FinalizedTransform);
+	FinalizedTransform.SetRotation(FinalRotation.Quaternion());
 
 	// Align to normal
-	if (PlacementSettings->bAllowAlignToNormal)
+	if (PlacementSettings->bAlignToNormal)
 	{
 		FRotator AlignRotation = InNormal.Rotation();
 		// Static meshes are authored along the vertical axis rather than the X axis, so we add 90 degrees to the static mesh's Pitch.
@@ -133,7 +124,7 @@ FTransform UPlacementBrushToolBase::GetFinalTransformFromHitLocationAndNormal(co
 		FinalizedTransform.SetRotation(FQuat(AlignRotation) * FinalizedTransform.GetRotation());
 	}
 
-	if (PlacementSettings->bAllowRandomScale)
+	if (PlacementSettings->bUseRandomScale)
 	{
 		// Until we have per object settings, just use a uniform scale, clamped from half to double size
 		FFloatInterval ScaleRange(0.5f, 2.0f);
@@ -144,6 +135,33 @@ FTransform UPlacementBrushToolBase::GetFinalTransformFromHitLocationAndNormal(co
 	return FinalizedTransform;
 }
 
+FRotator UPlacementBrushToolBase::GetFinalRotation(const FTransform& InTransform)
+{
+	FRotator UpdatedRotation = InTransform.Rotator();
+
+	if (!PlacementSettings.IsValid())
+	{
+		return UpdatedRotation;
+	}
+
+	if (PlacementSettings->bUseRandomXRotation)
+	{
+		UpdatedRotation.Roll = PlacementSettings->RandomRotationX.Interpolate(FMath::FRand());
+	}
+
+	if (PlacementSettings->bUseRandomYRotation)
+	{
+		UpdatedRotation.Pitch = PlacementSettings->RandomRotationY.Interpolate(FMath::FRand());
+	}
+
+	if (PlacementSettings->bUseRandomZRotation)
+	{
+		UpdatedRotation.Yaw = PlacementSettings->RandomRotationZ.Interpolate(FMath::FRand());
+	}
+
+	return UpdatedRotation;
+}
+
 TArray<FTypedElementHandle> UPlacementBrushToolBase::GetElementsInBrushRadius() const
 {
 	TArray<FTypedElementHandle> ElementHandles;
diff --git a/Engine/Source/Editor/AssetPlacementEdMode/Private/Tools/PlacementBrushToolBase.h b/Engine/Source/Editor/AssetPlacementEdMode/Private/Tools/PlacementBrushToolBase.h
--- a/Engine/Source/Editor/AssetPlacementEdMode/Private/Tools/PlacementBrushToolBase.h
+++ b/Engine/Source/Editor/AssetPlacementEdMode/Private/Tools/PlacementBrushToolBase.h
@@ -42,6 +42,7 @@ protected:
 	virtual double EstimateMaximumTargetDimension() override;
 	bool FindHitResultWithStartAndEndTraceVectors(FHitResult& OutHit, const FVector& TraceStart, const FVector& TraceEnd, float TraceRadius = 0.0f);
 	FTransform GetFinalTransformFromHitLocationAndNormal(const FVector& InLocation, const FVector& InNormal);
+	virtual FRotator GetFinalRotation(const FTransform& InTransform);
 	TArray<FTypedElementHandle> GetElementsInBrushRadius() const;
 
 	TWeakObjectPtr<UAssetPlacementSettings> PlacementSettings;
