Move shared settings object to mode subsystem.

Add additional single placement settings.

diff --git a/Engine/Source/Editor/AssetPlacementEdMode/AssetPlacementEdMode.Build.cs b/Engine/Source/Editor/AssetPlacementEdMode/AssetPlacementEdMode.Build.cs
--- a/Engine/Source/Editor/AssetPlacementEdMode/AssetPlacementEdMode.Build.cs
+++ b/Engine/Source/Editor/AssetPlacementEdMode/AssetPlacementEdMode.Build.cs
@@ -16,6 +16,7 @@ public class AssetPlacementEdMode : ModuleRules
 				"SlateCore",
                 "EditorStyle",
 				"EditorFramework",
+				"EditorSubsystem",
 				"UnrealEd",
 				"EditorInteractiveToolsFramework",
 				"InteractiveToolsFramework",
diff --git a/Engine/Source/Editor/AssetPlacementEdMode/Private/AssetPlacementEdMode.cpp b/Engine/Source/Editor/AssetPlacementEdMode/Private/AssetPlacementEdMode.cpp
--- a/Engine/Source/Editor/AssetPlacementEdMode/Private/AssetPlacementEdMode.cpp
+++ b/Engine/Source/Editor/AssetPlacementEdMode/Private/AssetPlacementEdMode.cpp
@@ -5,6 +5,7 @@
 #include "InteractiveToolManager.h"
 #include "AssetPlacementEdModeStyle.h"
 #include "AssetPlacementEdModeCommands.h"
+#include "AssetPlacementSettings.h"
 #include "EdModeInteractiveToolsContext.h"
 #include "EditorModeManager.h"
 #include "InstancedFoliageActor.h"
@@ -25,6 +26,7 @@
 #include "Elements/Interfaces/TypedElementObjectInterface.h"
 
 #include "Settings/LevelEditorMiscSettings.h"
+#include "Modes/PlacementModeSubsystem.h"
 
 #include "Subsystems/EditorActorSubsystem.h"
 #include "InstancedFoliageActor.h"
@@ -41,7 +43,6 @@ UAssetPlacementEdMode::UAssetPlacementEdMode()
 		FSlateIcon(FAssetPlacementEdModeStyle::GetStyleSetName(), "PlacementBrowser.ShowAllContent"),
 		MoveTemp(IsEnabledAttr));
 
-	SettingsClass = UAssetPlacementSettings::StaticClass();
 	bIsInSelectionTool = false;
 }
 
@@ -51,31 +52,25 @@ UAssetPlacementEdMode::~UAssetPlacementEdMode()
 
 void UAssetPlacementEdMode::Enter()
 {
+	// Set the settings object before we call Super::Enter, so that it's available for hooking up in the toolkit
+	SettingsObject = const_cast<UAssetPlacementSettings*>(GEditor->GetEditorSubsystem<UPlacementModeSubsystem>()->GetModeSettingsObject().Get());
+
 	Super::Enter();
 
 	const FAssetPlacementEdModeCommands& PlacementModeCommands = FAssetPlacementEdModeCommands::Get();
 	RegisterTool(PlacementModeCommands.Select, UPlacementModeSelectTool::ToolName, NewObject<UPlacementModeSelectToolBuilder>(this));
+	RegisterTool(PlacementModeCommands.Place, UPlacementModePlacementTool::ToolName, NewObject<UPlacementModePlacementToolBuilder>(this));
+	RegisterTool(PlacementModeCommands.LassoSelect, UPlacementModeLassoSelectTool::ToolName, NewObject<UPlacementModeLassoSelectToolBuilder>(this));
+	RegisterTool(PlacementModeCommands.PlaceSingle, UPlacementModePlaceSingleTool::ToolName, NewObject<UPlacementModePlaceSingleToolBuilder>(this));
+	RegisterTool(PlacementModeCommands.Erase, UPlacementModeEraseTool::ToolName, NewObject<UPlacementModeEraseToolBuilder>(this));
 
-	UPlacementModePlacementToolBuilder* PlaceToolBuilder = NewObject<UPlacementModePlacementToolBuilder>(this);
-	PlaceToolBuilder->PlacementSettings = SettingsObjectAsPlacementSettings;
-	RegisterTool(PlacementModeCommands.Place, UPlacementModePlacementTool::ToolName, PlaceToolBuilder);
-
-	UPlacementModeLassoSelectToolBuilder* LassoToolBuilder = NewObject<UPlacementModeLassoSelectToolBuilder>(this);
-	LassoToolBuilder->PlacementSettings = Cast<UAssetPlacementSettings>(SettingsObject);
-	RegisterTool(PlacementModeCommands.LassoSelect, UPlacementModeLassoSelectTool::ToolName, LassoToolBuilder);
-
-	UPlacementModePlaceSingleToolBuilder* SinglePlaceToolBuilder = NewObject<UPlacementModePlaceSingleToolBuilder>(this);
-	SinglePlaceToolBuilder->PlacementSettings = SettingsObjectAsPlacementSettings;
-	RegisterTool(PlacementModeCommands.PlaceSingle, UPlacementModePlaceSingleTool::ToolName, SinglePlaceToolBuilder);
-
-	UPlacementModeEraseToolBuilder* EraseToolBuilder = NewObject<UPlacementModeEraseToolBuilder>(this);
-	EraseToolBuilder->PlacementSettings = SettingsObjectAsPlacementSettings;
-	RegisterTool(PlacementModeCommands.Erase, UPlacementModeEraseTool::ToolName, EraseToolBuilder);
-
+	// Enable the select tool by default.
+	// Disable undo tracking so that we can't accidentally undo ourselves out of the select mode and into an invalid state.
 	GetToolManager()->ConfigureChangeTrackingMode(EToolChangeTrackingMode::NoChangeTracking);
 	ToolsContext->StartTool(UPlacementModeSelectTool::ToolName);
 	GetToolManager()->ConfigureChangeTrackingMode(EToolChangeTrackingMode::FullUndoRedo);
 
+	// Stash the current editor selection, since this mode will modify it.
 	Owner->StoreSelection(AssetPlacementEdModeID);
 }
 
@@ -90,6 +85,7 @@ void UAssetPlacementEdMode::Exit()
 		IFA->SelectInstance(nullptr, 0, bAppendSelection);
 	}
 
+	// Restore the selection to the original state.
 	Owner->RestoreSelection(AssetPlacementEdModeID);
 
 	SettingsObjectAsPlacementSettings.Reset();
@@ -138,7 +134,7 @@ bool UAssetPlacementEdMode::IsSelectionAllowed(AActor* InActor, bool bInSelectio
 
 	// And we need to have a valid palette item.
 	FTypedElementHandle ActorHandle = UEngineElementsLibrary::AcquireEditorActorElementHandle(InActor);
-	return DoesPaletteSupportElement(ActorHandle, SettingsObjectAsPlacementSettings->PaletteItems);
+	return GEditor->GetEditorSubsystem<UPlacementModeSubsystem>()->DoesCurrentPaletteSupportElement(ActorHandle);
 }
 
 void UAssetPlacementEdMode::OnToolStarted(UInteractiveToolManager* Manager, UInteractiveTool* Tool)
@@ -197,34 +193,6 @@ bool UAssetPlacementEdMode::IsEnabled()
 	return GetDefault<ULevelEditorMiscSettings>()->bEnableAssetPlacementMode;
 }
 
-bool UAssetPlacementEdMode::DoesPaletteSupportElement(const FTypedElementHandle& InElementToCheck, const TArray<FPaletteItem>& InPaletteToCheck)
-{
-	if (TTypedElement<UTypedElementAssetDataInterface> AssetDataInterface = UTypedElementRegistry::GetInstance()->GetElement<UTypedElementAssetDataInterface>(InElementToCheck))
-	{
-		TArray<FAssetData> ReferencedAssetDatas = AssetDataInterface.GetAllReferencedAssetDatas();
-		for (const FPaletteItem& Item : InPaletteToCheck)
-		{
-			if (ReferencedAssetDatas.Find(Item.AssetData) != INDEX_NONE)
-			{
-				return true;
-			}
-		}
-	}
-
-	// The current implementation of the asset data interface for actors requires that individual actors report on assets contained within their components.
-	// Not all actors do this reliably, so additionally check the supplied factory for a match. 
-	for (const FPaletteItem& Item : InPaletteToCheck)
-	{
-		FAssetData FoundAssetDataFromFactory = Item.FactoryOverride->GetAssetDataFromElementHandle(InElementToCheck);
-		if (FoundAssetDataFromFactory == Item.AssetData)
-		{
-			return true;
-		}
-	}
-
-	return false;
-}
-
 void UAssetPlacementEdMode::DeleteSelection()
 {
 	GetToolManager()->BeginUndoTransaction(LOCTEXT("PlacementDeleteAllSelected", "Delete Selected Assets"));
diff --git a/Engine/Source/Editor/AssetPlacementEdMode/Private/Modes/PlacementModeSubsystem.cpp b/Engine/Source/Editor/AssetPlacementEdMode/Private/Modes/PlacementModeSubsystem.cpp
new file mode 100644
--- /dev/null
+++ b/Engine/Source/Editor/AssetPlacementEdMode/Private/Modes/PlacementModeSubsystem.cpp
@@ -0,0 +1,59 @@
+// Copyright Epic Games, Inc. All Rights Reserved.
+
+#include "Modes/PlacementModeSubsystem.h"
+
+#include "AssetPlacementSettings.h"
+#include "Elements/Framework/TypedElementHandle.h"
+#include "Elements/Framework/TypedElementRegistry.h"
+#include "Elements/Interfaces/TypedElementAssetDataInterface.h"
+#include "Factories/AssetFactoryInterface.h"
+
+void UPlacementModeSubsystem::Initialize(FSubsystemCollectionBase& Collection)
+{
+	ModeSettings = NewObject<UAssetPlacementSettings>(this);
+	ModeSettings->LoadConfig();
+}
+
+void UPlacementModeSubsystem::Deinitialize()
+{
+	ModeSettings->SaveConfig();
+	ModeSettings = nullptr;
+}
+
+TWeakObjectPtr<const UAssetPlacementSettings> UPlacementModeSubsystem::GetModeSettingsObject() const
+{
+	return ModeSettings;
+}
+
+bool UPlacementModeSubsystem::DoesCurrentPaletteSupportElement(const FTypedElementHandle& InElementToCheck) const
+{
+	if (!ModeSettings)
+	{
+		return false;
+	}
+
+	if (TTypedElement<UTypedElementAssetDataInterface> AssetDataInterface = UTypedElementRegistry::GetInstance()->GetElement<UTypedElementAssetDataInterface>(InElementToCheck))
+	{
+		TArray<FAssetData> ReferencedAssetDatas = AssetDataInterface.GetAllReferencedAssetDatas();
+		for (const FPaletteItem& Item : ModeSettings->PaletteItems)
+		{
+			if (ReferencedAssetDatas.Find(Item.AssetData) != INDEX_NONE)
+			{
+				return true;
+			}
+		}
+	}
+
+	// The current implementation of the asset data interface for actors requires that individual actors report on assets contained within their components.
+	// Not all actors do this reliably, so additionally check the supplied factory for a match. 
+	for (const FPaletteItem& Item : ModeSettings->PaletteItems)
+	{
+		FAssetData FoundAssetDataFromFactory = Item.FactoryOverride->GetAssetDataFromElementHandle(InElementToCheck);
+		if (FoundAssetDataFromFactory == Item.AssetData)
+		{
+			return true;
+		}
+	}
+
+	return false;
+}
diff --git a/Engine/Source/Editor/AssetPlacementEdMode/Private/Modes/PlacementModeSubsystem.h b/Engine/Source/Editor/AssetPlacementEdMode/Private/Modes/PlacementModeSubsystem.h
new file mode 100644
--- /dev/null
+++ b/Engine/Source/Editor/AssetPlacementEdMode/Private/Modes/PlacementModeSubsystem.h
@@ -0,0 +1,37 @@
+// Copyright Epic Games, Inc. All Rights Reserved.
+
+#pragma once
+
+#include "EditorSubsystem.h"
+#include "UObject/ObjectPtr.h"
+
+#include "PlacementModeSubsystem.generated.h"
+
+struct FTypedElementHandle;
+class UAssetPlacementSettings;
+
+UCLASS(Transient)
+class UPlacementModeSubsystem : public UEditorSubsystem
+{
+	GENERATED_BODY()
+
+public:
+	// USubsystem Interface
+	virtual void Initialize(FSubsystemCollectionBase& Collection) override;
+	virtual void Deinitialize() override;
+	// End USubsystem Interface
+
+	// @returns the settings object for the mode for sharing across all tools and tool builders.
+	TWeakObjectPtr<const UAssetPlacementSettings> GetModeSettingsObject() const;
+
+	/**
+	 * Verifies if the given element handle is supported by the current mode settings' palette.
+	 *
+	 * @returns true if the element can be placed by the mode.
+	 */
+	bool DoesCurrentPaletteSupportElement(const FTypedElementHandle& InElementToCheck) const;
+
+protected:
+	UPROPERTY()
+	TObjectPtr<UAssetPlacementSettings> ModeSettings;
+};
diff --git a/Engine/Source/Editor/AssetPlacementEdMode/Private/Tools/PlacementBrushToolBase.cpp b/Engine/Source/Editor/AssetPlacementEdMode/Private/Tools/PlacementBrushToolBase.cpp
--- a/Engine/Source/Editor/AssetPlacementEdMode/Private/Tools/PlacementBrushToolBase.cpp
+++ b/Engine/Source/Editor/AssetPlacementEdMode/Private/Tools/PlacementBrushToolBase.cpp
@@ -17,18 +17,17 @@
 #include "AssetPlacementEdMode.h"
 #include "ActorPartition/ActorPartitionSubsystem.h"
 #include "Editor.h"
+#include "Modes/PlacementModeSubsystem.h"
 
 bool UPlacementToolBuilderBase::CanBuildTool(const FToolBuilderState& SceneState) const
-{
+{	
+	TWeakObjectPtr<const UAssetPlacementSettings> PlacementSettings = GEditor->GetEditorSubsystem<UPlacementModeSubsystem>()->GetModeSettingsObject();
 	return PlacementSettings.IsValid() && PlacementSettings->PaletteItems.Num();
 }
 
 UInteractiveTool* UPlacementToolBuilderBase::BuildTool(const FToolBuilderState& SceneState) const
 {
-	UPlacementBrushToolBase* NewTool = FactoryToolInstance(SceneState.ToolManager);
-	NewTool->PlacementSettings = PlacementSettings;
-
-	return NewTool;
+	return FactoryToolInstance(SceneState.ToolManager);
 }
 
 bool UPlacementBrushToolBase::HitTest(const FRay& Ray, FHitResult& OutHit)
@@ -54,15 +53,17 @@ bool UPlacementBrushToolBase::FindHitResultWithStartAndEndTraceVectors(FHitResul
 	UWorld* EditingWorld = GetToolManager()->GetWorld();
 	constexpr TCHAR NAME_PlacementBrushTool[] = TEXT("PlacementBrushTool");
 
-	auto FilterFunc = [this](const UPrimitiveComponent* InComponent) {
-		if (InComponent && this->PlacementSettings.IsValid())
+	TWeakObjectPtr<const UAssetPlacementSettings> PlacementSettings = GEditor->GetEditorSubsystem<UPlacementModeSubsystem>()->GetModeSettingsObject();
+
+	auto FilterFunc = [PlacementSettings](const UPrimitiveComponent* InComponent) {
+		if (InComponent && PlacementSettings.IsValid())
 		{
 			bool bFoliageOwned = InComponent->GetOwner() && FFoliageHelper::IsOwnedByFoliage(InComponent->GetOwner());
-			const bool bAllowLandscape = this->PlacementSettings->bLandscape;
-			const bool bAllowStaticMesh = this->PlacementSettings->bStaticMeshes;
-			const bool bAllowBSP = this->PlacementSettings->bBSP;
-			const bool bAllowFoliage = this->PlacementSettings->bFoliage;
-			const bool bAllowTranslucent = this->PlacementSettings->bTranslucent;
+			const bool bAllowLandscape = PlacementSettings->bLandscape;
+			const bool bAllowStaticMesh = PlacementSettings->bStaticMeshes;
+			const bool bAllowBSP = PlacementSettings->bBSP;
+			const bool bAllowFoliage = PlacementSettings->bFoliage;
+			const bool bAllowTranslucent = PlacementSettings->bTranslucent;
 
 			// Whitelist
 			bool bAllowed =
@@ -87,6 +88,7 @@ FTransform UPlacementBrushToolBase::GetFinalTransformFromHitLocationAndNormal(co
 {
 	FTransform FinalizedTransform(InLocation);
 
+	TWeakObjectPtr<const UAssetPlacementSettings> PlacementSettings = GEditor->GetEditorSubsystem<UPlacementModeSubsystem>()->GetModeSettingsObject();
 	if (!PlacementSettings.IsValid())
 	{
 		return FinalizedTransform;
@@ -139,6 +141,7 @@ FRotator UPlacementBrushToolBase::GetFinalRotation(const FTransform& InTransform
 {
 	FRotator UpdatedRotation = InTransform.Rotator();
 
+	TWeakObjectPtr<const UAssetPlacementSettings> PlacementSettings = GEditor->GetEditorSubsystem<UPlacementModeSubsystem>()->GetModeSettingsObject();
 	if (!PlacementSettings.IsValid())
 	{
 		return UpdatedRotation;
@@ -198,7 +201,7 @@ TArray<FTypedElementHandle> UPlacementBrushToolBase::GetElementsInBrushRadius()
 		if (HitActor)
 		{
 			FTypedElementHandle ActorHandle = UEngineElementsLibrary::AcquireEditorActorElementHandle(HitActor);
-			if (UAssetPlacementEdMode::DoesPaletteSupportElement(ActorHandle, PlacementSettings->PaletteItems))
+			if (GEditor->GetEditorSubsystem<UPlacementModeSubsystem>()->DoesCurrentPaletteSupportElement(ActorHandle))
 			{
 				ElementHandles.Emplace(ActorHandle);
 			}
@@ -215,7 +218,7 @@ TArray<FTypedElementHandle> UPlacementBrushToolBase::GetElementsInBrushRadius()
 			for (const auto& FoliageInfo : FoliageActor->GetFoliageInfos())
 			{
 				FTypedElementHandle SourceObjectHandle = UEngineElementsLibrary::AcquireEditorObjectElementHandle(FoliageInfo.Key->GetSource());
-				if (UAssetPlacementEdMode::DoesPaletteSupportElement(SourceObjectHandle, PlacementSettings->PaletteItems))
+				if (GEditor->GetEditorSubsystem<UPlacementModeSubsystem>()->DoesCurrentPaletteSupportElement(SourceObjectHandle))
 				{
 					TArray<int32> Instances;
 					FSphere SphereToCheck(LastBrushStamp.WorldPosition, LastBrushStamp.Radius);
diff --git a/Engine/Source/Editor/AssetPlacementEdMode/Private/Tools/PlacementBrushToolBase.h b/Engine/Source/Editor/AssetPlacementEdMode/Private/Tools/PlacementBrushToolBase.h
--- a/Engine/Source/Editor/AssetPlacementEdMode/Private/Tools/PlacementBrushToolBase.h
+++ b/Engine/Source/Editor/AssetPlacementEdMode/Private/Tools/PlacementBrushToolBase.h
@@ -20,9 +20,6 @@ public:
 	virtual bool CanBuildTool(const FToolBuilderState& SceneState) const override;
 	virtual UInteractiveTool* BuildTool(const FToolBuilderState& SceneState) const override;
 
-	UPROPERTY()
-	TWeakObjectPtr<UAssetPlacementSettings> PlacementSettings;
-
 protected:
 	virtual UPlacementBrushToolBase* FactoryToolInstance(UObject* Outer) const PURE_VIRTUAL(UPlacementToolBuilderBase::FactoryToolInstance, return nullptr; );
 };
@@ -44,6 +41,4 @@ protected:
 	FTransform GetFinalTransformFromHitLocationAndNormal(const FVector& InLocation, const FVector& InNormal);
 	virtual FRotator GetFinalRotation(const FTransform& InTransform);
 	TArray<FTypedElementHandle> GetElementsInBrushRadius() const;
-
-	TWeakObjectPtr<UAssetPlacementSettings> PlacementSettings;
 };
diff --git a/Engine/Source/Editor/AssetPlacementEdMode/Private/Tools/PlacementEraseTool.cpp b/Engine/Source/Editor/AssetPlacementEdMode/Private/Tools/PlacementEraseTool.cpp
--- a/Engine/Source/Editor/AssetPlacementEdMode/Private/Tools/PlacementEraseTool.cpp
+++ b/Engine/Source/Editor/AssetPlacementEdMode/Private/Tools/PlacementEraseTool.cpp
@@ -13,6 +13,7 @@
 #include "InstancedFoliageActor.h"
 #include "AssetPlacementEdMode.h"
 #include "AssetPlacementSettings.h"
+#include "Modes/PlacementModeSubsystem.h"
 
 constexpr TCHAR UPlacementModeEraseTool::ToolName[];
 
@@ -60,7 +61,7 @@ void UPlacementModeEraseTool::OnTick(float DeltaTime)
 				FoliageActor->ForEachFoliageInfo([this](UFoliageType* FoliageType, FFoliageInfo& FoliageInfo)
 				{
 					FTypedElementHandle SourceObjectHandle = UEngineElementsLibrary::AcquireEditorObjectElementHandle(FoliageType->GetSource());
-					if (UAssetPlacementEdMode::DoesPaletteSupportElement(SourceObjectHandle, PlacementSettings->PaletteItems))
+					if (GEditor->GetEditorSubsystem<UPlacementModeSubsystem>()->DoesCurrentPaletteSupportElement(SourceObjectHandle))
 					{
 						TArray<int32> Instances;
 						FSphere SphereToCheck(LastBrushStamp.WorldPosition, LastBrushStamp.Radius);
diff --git a/Engine/Source/Editor/AssetPlacementEdMode/Private/Tools/PlacementLassoSelectTool.cpp b/Engine/Source/Editor/AssetPlacementEdMode/Private/Tools/PlacementLassoSelectTool.cpp
--- a/Engine/Source/Editor/AssetPlacementEdMode/Private/Tools/PlacementLassoSelectTool.cpp
+++ b/Engine/Source/Editor/AssetPlacementEdMode/Private/Tools/PlacementLassoSelectTool.cpp
@@ -14,6 +14,7 @@
 #include "UObject/Object.h"
 #include "BaseBehaviors/KeyAsModifierInputBehavior.h"
 #include "Elements/Framework/TypedElementSelectionSet.h"
+#include "Modes/PlacementModeSubsystem.h"
 
 constexpr TCHAR UPlacementModeLassoSelectTool::ToolName[];
 
@@ -56,7 +57,7 @@ void UPlacementModeLassoSelectTool::OnTick(float DeltaTime)
 				FoliageActor->ForEachFoliageInfo([this, bSelectElements](UFoliageType* InFoliageType, FFoliageInfo& InFoliageInfo)
 				{
 					FTypedElementHandle SourceObjectHandle = UEngineElementsLibrary::AcquireEditorObjectElementHandle(InFoliageType->GetSource());
-					if (UAssetPlacementEdMode::DoesPaletteSupportElement(SourceObjectHandle, PlacementSettings->PaletteItems))
+					if (GEditor->GetEditorSubsystem<UPlacementModeSubsystem>()->DoesCurrentPaletteSupportElement(SourceObjectHandle))
 					{
 						TArray<int32> Instances;
 						FSphere SphereToCheck(LastBrushStamp.WorldPosition, LastBrushStamp.Radius);
diff --git a/Engine/Source/Editor/AssetPlacementEdMode/Private/Tools/PlacementPlaceSingleTool.cpp b/Engine/Source/Editor/AssetPlacementEdMode/Private/Tools/PlacementPlaceSingleTool.cpp
--- a/Engine/Source/Editor/AssetPlacementEdMode/Private/Tools/PlacementPlaceSingleTool.cpp
+++ b/Engine/Source/Editor/AssetPlacementEdMode/Private/Tools/PlacementPlaceSingleTool.cpp
@@ -7,6 +7,7 @@
 #include "ScopedTransaction.h"
 #include "Editor/EditorEngine.h"
 #include "Editor.h"
+#include "Modes/PlacementModeSubsystem.h"
 
 constexpr TCHAR UPlacementModePlaceSingleTool::ToolName[];
 
@@ -15,10 +16,45 @@ UPlacementBrushToolBase* UPlacementModePlaceSingleToolBuilder::FactoryToolInstan
 	return NewObject<UPlacementModePlaceSingleTool>(Outer);
 }
 
+bool UPlacementModePlaceSingleSettings::CanEditChange(const FProperty* Property) const
+{
+	if (!Super::CanEditChange(Property))
+	{
+		return false;
+	}
+
+	const FName PropertyName = Property->GetFName();
+	if (PropertyName == GET_MEMBER_NAME_CHECKED(UPlacementModePlaceSingleSettings, bInvertCursorAxis))
+	{
+		return bAlignToCursor;
+	}
+
+	return true;
+}
+
+void UPlacementModePlaceSingleTool::Setup()
+{
+	Super::Setup();
+
+	SingleToolSettings = NewObject<UPlacementModePlaceSingleSettings>(this);
+	SingleToolSettings->LoadConfig();
+	AddToolPropertySource(SingleToolSettings);
+}
+
+void UPlacementModePlaceSingleTool::Shutdown(EToolShutdownType ShutdownType)
+{
+	Super::Shutdown(ShutdownType);
+
+	SingleToolSettings->SaveConfig();
+	RemoveToolPropertySource(SingleToolSettings);
+	SingleToolSettings = nullptr;
+}
+
 void UPlacementModePlaceSingleTool::OnEndDrag(const FRay& Ray)
 {
 	Super::OnEndDrag(Ray);
 
+	TWeakObjectPtr<const UAssetPlacementSettings> PlacementSettings = GEditor->GetEditorSubsystem<UPlacementModeSubsystem>()->GetModeSettingsObject();
 	if (PlacementSettings.IsValid() && PlacementSettings->PaletteItems.Num())
 	{
 		int32 ItemIndex = FMath::RandHelper(PlacementSettings->PaletteItems.Num());
@@ -44,3 +80,8 @@ void UPlacementModePlaceSingleTool::OnEndDrag(const FRay& Ray)
 		}
 	}
 }
+
+FRotator UPlacementModePlaceSingleTool::GetFinalRotation(const FTransform& InTransform)
+{
+	return InTransform.Rotator();
+}
diff --git a/Engine/Source/Editor/AssetPlacementEdMode/Private/Tools/PlacementPlaceSingleTool.h b/Engine/Source/Editor/AssetPlacementEdMode/Private/Tools/PlacementPlaceSingleTool.h
--- a/Engine/Source/Editor/AssetPlacementEdMode/Private/Tools/PlacementPlaceSingleTool.h
+++ b/Engine/Source/Editor/AssetPlacementEdMode/Private/Tools/PlacementPlaceSingleTool.h
@@ -16,6 +16,36 @@ protected:
 	virtual UPlacementBrushToolBase* FactoryToolInstance(UObject* Outer) const override;
 };
 
+UCLASS(config = EditorPerUserSettings)
+class UPlacementModePlaceSingleSettings : public UInteractiveToolPropertySet
+{
+	GENERATED_BODY()
+
+public:
+	UPROPERTY(config, EditAnywhere, Category = "Single Place Tool", meta = (InlineEditConditionToggle))
+	bool bAlignToCursor = false;
+
+	UPROPERTY(config, EditAnywhere, Category = "Single Place Tool", meta = (DisplayName = "Align to Cursor", EditCondition = "bAlignToCursor"))
+	TEnumAsByte<EAxis::Type> AxisToAlignWithCursor = EAxis::Type::X;
+
+	UPROPERTY(config, EditAnywhere, Category = "Single Place Tool")
+	bool bInvertCursorAxis = false;
+
+	UPROPERTY(config, EditAnywhere, Category = "Single Place Tool")
+	bool bSnapToGridX = false;
+
+	UPROPERTY(config, EditAnywhere, Category = "Single Place Tool")
+	bool bSnapToGridY = false;
+
+	UPROPERTY(config, EditAnywhere, Category = "Single Place Tool")
+	bool bSnapToGridZ = false;
+
+	UPROPERTY(config, EditAnywhere, Category = "Single Place Tool")
+	float ScrollWheelOffsetScale = 0.05f;
+
+	virtual bool CanEditChange(const FProperty* Property) const override;
+};
+
 UCLASS(MinimalAPI)
 class UPlacementModePlaceSingleTool : public UPlacementClickDragToolBase
 {
@@ -24,5 +54,11 @@ class UPlacementModePlaceSingleTool : public UPlacementClickDragToolBase
 public:
 	constexpr static TCHAR ToolName[] = TEXT("PlaceSingleTool");
 
+	virtual void Setup() override;
+	virtual void Shutdown(EToolShutdownType ShutdownType) override;
 	virtual void OnEndDrag(const FRay& Ray) override;
+
+protected:
+	virtual FRotator GetFinalRotation(const FTransform& InTransform) override;
+	TObjectPtr<UPlacementModePlaceSingleSettings> SingleToolSettings;
 };
diff --git a/Engine/Source/Editor/AssetPlacementEdMode/Private/Tools/PlacementPlaceTool.cpp b/Engine/Source/Editor/AssetPlacementEdMode/Private/Tools/PlacementPlaceTool.cpp
--- a/Engine/Source/Editor/AssetPlacementEdMode/Private/Tools/PlacementPlaceTool.cpp
+++ b/Engine/Source/Editor/AssetPlacementEdMode/Private/Tools/PlacementPlaceTool.cpp
@@ -8,6 +8,7 @@
 #include "Editor.h"
 #include "InstancedFoliage.h"
 #include "InteractiveToolManager.h"
+#include "Modes/PlacementModeSubsystem.h"
 
 constexpr TCHAR UPlacementModePlacementTool::ToolName[];
 
@@ -37,6 +38,7 @@ void UPlacementModePlacementTool::OnTick(float DeltaTime)
 		return;
 	}
 
+	TWeakObjectPtr<const UAssetPlacementSettings> PlacementSettings = GEditor->GetEditorSubsystem<UPlacementModeSubsystem>()->GetModeSettingsObject();
 	if (BrushProperties && PlacementSettings.IsValid() && PlacementSettings->PaletteItems.Num())
 	{
 		FVector BrushLocation = LastBrushStamp.WorldPosition;
diff --git a/Engine/Source/Editor/AssetPlacementEdMode/Public/AssetPlacementEdMode.h b/Engine/Source/Editor/AssetPlacementEdMode/Public/AssetPlacementEdMode.h
--- a/Engine/Source/Editor/AssetPlacementEdMode/Public/AssetPlacementEdMode.h
+++ b/Engine/Source/Editor/AssetPlacementEdMode/Public/AssetPlacementEdMode.h
@@ -5,10 +5,11 @@
 #include "CoreMinimal.h"
 #include "Tools/LegacyEdModeWidgetHelpers.h"
 #include "Elements/Framework/TypedElementHandle.h"
-#include "AssetPlacementSettings.h"
 
 #include "AssetPlacementEdMode.generated.h"
 
+class UAssetPlacementSettings;
+
 UCLASS()
 class UAssetPlacementEdMode : public UBaseLegacyWidgetEdMode
 {
@@ -40,7 +41,6 @@ public:
 	//////////////////
 
 	static bool IsEnabled();
-	static bool DoesPaletteSupportElement(const FTypedElementHandle& InElementToCheck, const TArray<FPaletteItem>& InPaletteToCheck);
 
 protected:
 	enum class ESelectMode
