Placement Single tool improvements:
* Preview actor to visualize placement.
* Ctrl + click to rotate to another item in the palette and regenerate rotation and scale.
* Configurable normal axis alignment.

Also fixed a crash on exiting the mode due to lingering mode toolkit references in the level editor.

diff --git a/Engine/Source/Editor/AssetPlacementEdMode/Private/Tools/PlacementBrushToolBase.cpp b/Engine/Source/Editor/AssetPlacementEdMode/Private/Tools/PlacementBrushToolBase.cpp
--- a/Engine/Source/Editor/AssetPlacementEdMode/Private/Tools/PlacementBrushToolBase.cpp
+++ b/Engine/Source/Editor/AssetPlacementEdMode/Private/Tools/PlacementBrushToolBase.cpp
@@ -18,6 +18,7 @@
 #include "ActorPartition/ActorPartitionSubsystem.h"
 #include "Editor.h"
 #include "Modes/PlacementModeSubsystem.h"
+#include "ActorFactories/ActorFactory.h"
 
 bool UPlacementToolBuilderBase::CanBuildTool(const FToolBuilderState& SceneState) const
 {	
@@ -97,34 +98,7 @@ FTransform UPlacementBrushToolBase::GetFinalTransformFromHitLocationAndNormal(co
 	// Update the rotation if we need to
 	FRotator FinalRotation = GetFinalRotation(FinalizedTransform);
 	FinalizedTransform.SetRotation(FinalRotation.Quaternion());
-
-	// Align to normal
-	if (PlacementSettings->bAlignToNormal)
-	{
-		FRotator AlignRotation = InNormal.Rotation();
-		// Static meshes are authored along the vertical axis rather than the X axis, so we add 90 degrees to the static mesh's Pitch.
-		AlignRotation.Pitch -= 90.f;
-		// Clamp its value inside +/- one rotation
-		AlignRotation.Pitch = FRotator::NormalizeAxis(AlignRotation.Pitch);
-
-		// limit the maximum pitch angle if it's > 0.
-		// For now, just set the align max angle to a constant value, until it can be pulled from per object settings
-		constexpr float AlignMaxAngle = 0.0f;
-		if (AlignMaxAngle > 0.f)
-		{
-			int32 MaxPitch = AlignMaxAngle;
-			if (AlignRotation.Pitch > MaxPitch)
-			{
-				AlignRotation.Pitch = MaxPitch;
-			}
-			else if (AlignRotation.Pitch < -MaxPitch)
-			{
-				AlignRotation.Pitch = -MaxPitch;
-			}
-		}
-
-		FinalizedTransform.SetRotation(FQuat(AlignRotation) * FinalizedTransform.GetRotation());
-	}
+	FinalizedTransform.NormalizeRotation();
 
 	if (PlacementSettings->bUseRandomScale)
 	{
@@ -162,9 +136,45 @@ FRotator UPlacementBrushToolBase::GetFinalRotation(const FTransform& InTransform
 		UpdatedRotation.Yaw = PlacementSettings->RandomRotationZ.Interpolate(FMath::FRand());
 	}
 
+	// Cache this rotation before we start aligning to normal
+	LastGeneratedRotation = FQuat(UpdatedRotation);
+
+	// Align to normal
+	if (PlacementSettings->bAlignToNormal)
+	{
+		UpdatedRotation = UpdateRotationAlignedToBrushNormal(PlacementSettings->AxisToAlignWithNormal, PlacementSettings->bInvertNormalAxis).Rotator();
+	}
+
 	return UpdatedRotation;
 }
 
+FQuat UPlacementBrushToolBase::UpdateRotationAlignedToBrushNormal(EAxis::Type InAlignmentAxis, bool bInvertAxis)
+{
+	FVector AlignmentVector = FVector::ZeroVector;
+	switch (InAlignmentAxis)
+	{
+		case EAxis::Type::X:
+		{
+			AlignmentVector = bInvertAxis ? FVector::BackwardVector : FVector::ForwardVector;
+		}
+		break;
+
+		case EAxis::Type::Y:
+		{
+			AlignmentVector = bInvertAxis ? FVector::LeftVector : FVector::RightVector;
+		}
+		break;
+
+		case EAxis::Type::Z:
+		{
+			AlignmentVector = bInvertAxis ? FVector::DownVector : FVector::UpVector;
+		}
+		break;
+	}
+
+	return FindActorAlignmentRotation(LastGeneratedRotation, AlignmentVector, LastBrushStamp.WorldNormal, &LastAlignRotation);
+}
+
 TArray<FTypedElementHandle> UPlacementBrushToolBase::GetElementsInBrushRadius() const
 {
 	TArray<FTypedElementHandle> ElementHandles;
diff --git a/Engine/Source/Editor/AssetPlacementEdMode/Private/Tools/PlacementBrushToolBase.h b/Engine/Source/Editor/AssetPlacementEdMode/Private/Tools/PlacementBrushToolBase.h
--- a/Engine/Source/Editor/AssetPlacementEdMode/Private/Tools/PlacementBrushToolBase.h
+++ b/Engine/Source/Editor/AssetPlacementEdMode/Private/Tools/PlacementBrushToolBase.h
@@ -39,6 +39,14 @@ protected:
 	virtual double EstimateMaximumTargetDimension() override;
 	bool FindHitResultWithStartAndEndTraceVectors(FHitResult& OutHit, const FVector& TraceStart, const FVector& TraceEnd, float TraceRadius = 0.0f);
 	FTransform GetFinalTransformFromHitLocationAndNormal(const FVector& InLocation, const FVector& InNormal);
-	virtual FRotator GetFinalRotation(const FTransform& InTransform);
+
+	// Gets a random rotation, and aligns it, based on the placement settings.
+	FRotator GetFinalRotation(const FTransform& InTransform);
+
+	// Updates the last generated rotation to realign with the current brush position and normal.
+	FQuat UpdateRotationAlignedToBrushNormal(EAxis::Type InAlignmentAxis, bool bInvertAxis);
 	TArray<FTypedElementHandle> GetElementsInBrushRadius() const;
+
+	FQuat LastGeneratedRotation;
+	FQuat LastAlignRotation;
 };
diff --git a/Engine/Source/Editor/AssetPlacementEdMode/Private/Tools/PlacementClickDragToolBase.cpp b/Engine/Source/Editor/AssetPlacementEdMode/Private/Tools/PlacementClickDragToolBase.cpp
--- a/Engine/Source/Editor/AssetPlacementEdMode/Private/Tools/PlacementClickDragToolBase.cpp
+++ b/Engine/Source/Editor/AssetPlacementEdMode/Private/Tools/PlacementClickDragToolBase.cpp
@@ -8,7 +8,7 @@
 
 void UPlacementClickDragToolBase::Setup()
 {
-	UPlacementBrushToolBase::Setup();
+	Super::Setup();
 
 	RemoveToolPropertySource(BrushProperties);
 	BrushStampIndicator->bDrawRadiusCircle = false;
@@ -18,5 +18,5 @@ void UPlacementClickDragToolBase::Setup()
 
 double UPlacementClickDragToolBase::EstimateMaximumTargetDimension()
 {
-	return UBaseBrushTool::EstimateMaximumTargetDimension();
+	return Super::EstimateMaximumTargetDimension();
 }
\ No newline at end of file
diff --git a/Engine/Source/Editor/AssetPlacementEdMode/Private/Tools/PlacementPlaceSingleTool.cpp b/Engine/Source/Editor/AssetPlacementEdMode/Private/Tools/PlacementPlaceSingleTool.cpp
--- a/Engine/Source/Editor/AssetPlacementEdMode/Private/Tools/PlacementPlaceSingleTool.cpp
+++ b/Engine/Source/Editor/AssetPlacementEdMode/Private/Tools/PlacementPlaceSingleTool.cpp
@@ -8,6 +8,11 @@
 #include "Editor/EditorEngine.h"
 #include "Editor.h"
 #include "Modes/PlacementModeSubsystem.h"
+#include "InteractiveToolManager.h"
+#include "Elements/Framework/TypedElementRegistry.h"
+#include "Elements/Framework/TypedElementHandle.h"
+#include "Elements/Interfaces/TypedElementObjectInterface.h"
+#include "Elements/Interfaces/TypedElementWorldInterface.h"
 
 constexpr TCHAR UPlacementModePlaceSingleTool::ToolName[];
 
@@ -32,6 +37,13 @@ bool UPlacementModePlaceSingleSettings::CanEditChange(const FProperty* Property)
 	return true;
 }
 
+UPlacementModePlaceSingleTool::UPlacementModePlaceSingleTool() = default;
+UPlacementModePlaceSingleTool::~UPlacementModePlaceSingleTool() = default;
+UPlacementModePlaceSingleTool::UPlacementModePlaceSingleTool(FVTableHelper& Helper)
+	: Super(Helper)
+{
+}
+
 void UPlacementModePlaceSingleTool::Setup()
 {
 	Super::Setup();
@@ -43,45 +55,167 @@ void UPlacementModePlaceSingleTool::Setup()
 
 void UPlacementModePlaceSingleTool::Shutdown(EToolShutdownType ShutdownType)
 {
-	Super::Shutdown(ShutdownType);
+	DestroyPreviewElements();
 
 	SingleToolSettings->SaveConfig();
 	RemoveToolPropertySource(SingleToolSettings);
 	SingleToolSettings = nullptr;
+
+	Super::Shutdown(ShutdownType);
+}
+
+void UPlacementModePlaceSingleTool::OnClickPress(const FInputDeviceRay& PressPos)
+{
+	Super::OnClickPress(PressPos);
+
+	if (bCtrlToggle)
+	{
+		// If this was a ctrl click, regenerate the preview placement information, and destroy any active gizmos.
+		GeneratePreviewPlacementData(PressPos);
+	}
+	else
+	{
+		// Otherwise, and we're not currently editing something that was placed, then place the preview data.
+		// Todo - verify we don't get here if the gizmos are active
+		FPlacementOptions PlacementOptions;
+		PlacementOptions.bPreferBatchPlacement = true;
+		PlacementOptions.bPreferInstancedPlacement = true;
+
+		UPlacementSubsystem* PlacementSubsystem = GEditor->GetEditorSubsystem<UPlacementSubsystem>();
+		if (PlacementSubsystem)
+		{
+			// Update the level, just in case the preview moved us out of the current one.
+			PreviewPlacementInfo->PreferredLevel = GEditor->GetEditorWorldContext().World()->GetCurrentLevel();
+
+			FScopedTransaction Transaction(NSLOCTEXT("PlacementMode", "SinglePlaceAsset", "Place Single Asset"));
+			PlacementSubsystem->PlaceAsset(*PreviewPlacementInfo, PlacementOptions);
+		}
+
+		DestroyPreviewElements();
+	}
+}
+
+void UPlacementModePlaceSingleTool::OnBeginHover(const FInputDeviceRay& DevicePos)
+{
+	Super::OnBeginHover(DevicePos);
+
+	// Always regenerate the placement data when a hover sequence begins
+	GeneratePreviewPlacementData(DevicePos);
+}
+
+bool UPlacementModePlaceSingleTool::OnUpdateHover(const FInputDeviceRay& DevicePos)
+{
+	if (!Super::OnUpdateHover(DevicePos))
+	{
+		return false;
+	}
+
+	// Update the preview elements
+	UpdatePreviewElements(DevicePos);
+	return true;
 }
 
-void UPlacementModePlaceSingleTool::OnEndDrag(const FRay& Ray)
+void UPlacementModePlaceSingleTool::OnEndHover()
 {
-	Super::OnEndDrag(Ray);
+	Super::OnEndHover();
+
+	// remove the preview elements from the level
+	DestroyPreviewElements();
+}
+
+void UPlacementModePlaceSingleTool::GeneratePreviewPlacementData(const FInputDeviceRay& DevicePos)
+{
+	PreviewPlacementInfo.Reset();
 
 	TWeakObjectPtr<const UAssetPlacementSettings> PlacementSettings = GEditor->GetEditorSubsystem<UPlacementModeSubsystem>()->GetModeSettingsObject();
 	if (PlacementSettings.IsValid() && PlacementSettings->PaletteItems.Num())
 	{
 		int32 ItemIndex = FMath::RandHelper(PlacementSettings->PaletteItems.Num());
 		const FPaletteItem& ItemToPlace = PlacementSettings->PaletteItems[ItemIndex];
-		if (ItemToPlace.AssetData.IsValid())
+
+		PreviewPlacementInfo = MakeUnique<FAssetPlacementInfo>();
+		PreviewPlacementInfo->AssetToPlace = ItemToPlace.AssetData;
+		PreviewPlacementInfo->FactoryOverride = ItemToPlace.FactoryOverride;
+		PreviewPlacementInfo->FinalizedTransform = GetFinalTransformFromHitLocationAndNormal(LastBrushStamp.HitResult.ImpactPoint, LastBrushStamp.HitResult.ImpactNormal);
+		PreviewPlacementInfo->PreferredLevel = GEditor->GetEditorWorldContext().World()->GetCurrentLevel();
+	}
+}
+
+void UPlacementModePlaceSingleTool::CreatePreviewElements(const FInputDeviceRay& DevicePos)
+{
+	// Place the preview elements from our stored info
+	if (!PreviewPlacementInfo.IsValid())
+	{
+		GeneratePreviewPlacementData(DevicePos);
+	}
+
+	if (UPlacementSubsystem* PlacementSubsystem = GEditor->GetEditorSubsystem<UPlacementSubsystem>())
+	{
+		FPlacementOptions PlacementOptions;
+		PlacementOptions.bIsCreatingPreviewElements = true;
+
+		PreviewElements = PlacementSubsystem->PlaceAsset(*PreviewPlacementInfo, PlacementOptions);
+	}
+	
+	// Disable collision on any preview elements
+	for (FTypedElementHandle& PreviewElement : PreviewElements)
+	{
+		if (TTypedElement<UTypedElementWorldInterface> WorldInterfaceElement = UTypedElementRegistry::GetInstance()->GetElement<UTypedElementWorldInterface>(PreviewElement))
 		{
-			FAssetPlacementInfo PlacementInfo;
-			PlacementInfo.AssetToPlace = ItemToPlace.AssetData;
-			PlacementInfo.FactoryOverride = ItemToPlace.FactoryOverride;
-			PlacementInfo.FinalizedTransform = GetFinalTransformFromHitLocationAndNormal(LastBrushStamp.HitResult.ImpactPoint, LastBrushStamp.HitResult.ImpactNormal);
-			PlacementInfo.PreferredLevel = GEditor->GetEditorWorldContext().World()->GetCurrentLevel();
-
-			FPlacementOptions PlacementOptions;
-			PlacementOptions.bPreferBatchPlacement = true;
-			PlacementOptions.bPreferInstancedPlacement = true;
-
-			UPlacementSubsystem* PlacementSubsystem = GEditor->GetEditorSubsystem<UPlacementSubsystem>();
-			if (PlacementSubsystem)
-			{
-				FScopedTransaction Transaction(NSLOCTEXT("PlacementMode", "SinglePlaceAsset", "Place Single Asset"));
-				PlacementSubsystem->PlaceAsset(PlacementInfo, PlacementOptions);
-			}
+			WorldInterfaceElement.NotifyMovementStarted();
 		}
 	}
 }
 
-FRotator UPlacementModePlaceSingleTool::GetFinalRotation(const FTransform& InTransform)
+void UPlacementModePlaceSingleTool::UpdatePreviewElements(const FInputDeviceRay& DevicePos)
 {
-	return InTransform.Rotator();
+	if (PreviewElements.Num() == 0)
+	{
+		CreatePreviewElements(DevicePos);
+	}
+
+	if (PreviewElements.Num() == 0)
+	{
+		return;
+	}
+
+	// Update the location.
+	FTransform& FinalizedTransform = PreviewPlacementInfo->FinalizedTransform;
+	FinalizedTransform.SetLocation(LastBrushStamp.HitResult.ImpactPoint);
+
+	// Update the rotation.
+	TWeakObjectPtr<const UAssetPlacementSettings> PlacementSettings = GEditor->GetEditorSubsystem<UPlacementModeSubsystem>()->GetModeSettingsObject();
+	if (PlacementSettings.IsValid() && PlacementSettings->bAlignToNormal)
+	{
+		FQuat UpdatedRotation = UpdateRotationAlignedToBrushNormal(PlacementSettings->AxisToAlignWithNormal, PlacementSettings->bInvertNormalAxis);
+		FinalizedTransform.SetRotation(UpdatedRotation);
+	}
+
+	// Update the transform positions for the preview elements.
+	for (FTypedElementHandle& PreviewElement : PreviewElements)
+	{
+		if (TTypedElement<UTypedElementWorldInterface> WorldInterfaceElement = UTypedElementRegistry::GetInstance()->GetElement<UTypedElementWorldInterface>(PreviewElement))
+		{
+			WorldInterfaceElement.SetWorldTransform(FinalizedTransform);
+			WorldInterfaceElement.NotifyMovementOngoing();
+		}
+	}
+}
+
+void UPlacementModePlaceSingleTool::DestroyPreviewElements()
+{
+	FToolBuilderState SelectionState;
+	GetToolManager()->GetContextQueriesAPI()->GetCurrentSelectionState(SelectionState);
+	check(SelectionState.TypedElementSelectionSet.IsValid());	// Placement tools expect a valid selection set.
+
+	for (FTypedElementHandle& PreviewElement : PreviewElements)
+	{
+		if (TTypedElement<UTypedElementWorldInterface> WorldInterfaceElement = UTypedElementRegistry::GetInstance()->GetElement<UTypedElementWorldInterface>(PreviewElement))
+		{
+			WorldInterfaceElement.NotifyMovementEnded();
+			WorldInterfaceElement.DeleteElement(WorldInterfaceElement.GetOwnerWorld(), SelectionState.TypedElementSelectionSet.Get(), FTypedElementDeletionOptions());
+		}
+	}
+
+	PreviewElements.Empty();
 }
diff --git a/Engine/Source/Editor/AssetPlacementEdMode/Private/Tools/PlacementPlaceSingleTool.h b/Engine/Source/Editor/AssetPlacementEdMode/Private/Tools/PlacementPlaceSingleTool.h
--- a/Engine/Source/Editor/AssetPlacementEdMode/Private/Tools/PlacementPlaceSingleTool.h
+++ b/Engine/Source/Editor/AssetPlacementEdMode/Private/Tools/PlacementPlaceSingleTool.h
@@ -7,6 +7,8 @@
 
 #include "PlacementPlaceSingleTool.generated.h"
 
+struct FAssetPlacementInfo;
+
 UCLASS(Transient, MinimalAPI)
 class UPlacementModePlaceSingleToolBuilder : public UPlacementToolBuilderBase
 {
@@ -46,19 +48,35 @@ public:
 	virtual bool CanEditChange(const FProperty* Property) const override;
 };
 
-UCLASS(MinimalAPI)
+UCLASS(Transient)
 class UPlacementModePlaceSingleTool : public UPlacementClickDragToolBase
 {
 	GENERATED_BODY()
 
 public:
+	UPlacementModePlaceSingleTool();
+	~UPlacementModePlaceSingleTool();
+	UPlacementModePlaceSingleTool(FVTableHelper& Helper);
+	
 	constexpr static TCHAR ToolName[] = TEXT("PlaceSingleTool");
 
 	virtual void Setup() override;
 	virtual void Shutdown(EToolShutdownType ShutdownType) override;
-	virtual void OnEndDrag(const FRay& Ray) override;
+	virtual void OnClickPress(const FInputDeviceRay& Ray) override;
+
+	virtual void OnBeginHover(const FInputDeviceRay& DevicePos) override;
+	virtual bool OnUpdateHover(const FInputDeviceRay& DevicePos) override;
+	virtual void OnEndHover() override;
 
 protected:
-	virtual FRotator GetFinalRotation(const FTransform& InTransform) override;
+	void GeneratePreviewPlacementData(const FInputDeviceRay& DevicePos);
+	void CreatePreviewElements(const FInputDeviceRay& DevicePos);
+	void UpdatePreviewElements(const FInputDeviceRay& DevicePos);
+	void DestroyPreviewElements();
+
+	UPROPERTY()
 	TObjectPtr<UPlacementModePlaceSingleSettings> SingleToolSettings;
+
+	TUniquePtr<FAssetPlacementInfo> PreviewPlacementInfo;
+	TArray<FTypedElementHandle> PreviewElements;
 };
diff --git a/Engine/Source/Editor/EditorFramework/Private/Subsystems/PlacementSubsystem.cpp b/Engine/Source/Editor/EditorFramework/Private/Subsystems/PlacementSubsystem.cpp
--- a/Engine/Source/Editor/EditorFramework/Private/Subsystems/PlacementSubsystem.cpp
+++ b/Engine/Source/Editor/EditorFramework/Private/Subsystems/PlacementSubsystem.cpp
@@ -30,6 +30,9 @@ TArray<FTypedElementHandle> UPlacementSubsystem::PlaceAsset(const FAssetPlacemen
 
 TArray<FTypedElementHandle> UPlacementSubsystem::PlaceAssets(TArrayView<const FAssetPlacementInfo> InPlacementInfos, const FPlacementOptions& InPlacementOptions)
 {
+	bool bWasCreatingPreviewElements = bIsCreatingPreviewElements;
+	bIsCreatingPreviewElements = InPlacementOptions.bIsCreatingPreviewElements;
+
 	TArray<FTypedElementHandle> PlacedElements;
 	for (const FAssetPlacementInfo& PlacementInfo : InPlacementInfos)
 	{
@@ -59,6 +62,8 @@ TArray<FTypedElementHandle> UPlacementSubsystem::PlaceAssets(TArrayView<const FA
 		}
 	}
 
+	bIsCreatingPreviewElements = bWasCreatingPreviewElements;
+
 	return PlacedElements;
 }
 
@@ -75,6 +80,11 @@ TScriptInterface<IAssetFactoryInterface> UPlacementSubsystem::FindAssetFactoryFr
 	return nullptr;
 }
 
+bool UPlacementSubsystem::IsCreatingPreviewElements() const
+{
+	return bIsCreatingPreviewElements;
+}
+
 void UPlacementSubsystem::RegisterPlacementFactories()
 {
 	for (TObjectIterator<UClass> ObjectIt; ObjectIt; ++ObjectIt)
diff --git a/Engine/Source/Editor/EditorFramework/Public/Subsystems/PlacementSubsystem.h b/Engine/Source/Editor/EditorFramework/Public/Subsystems/PlacementSubsystem.h
--- a/Engine/Source/Editor/EditorFramework/Public/Subsystems/PlacementSubsystem.h
+++ b/Engine/Source/Editor/EditorFramework/Public/Subsystems/PlacementSubsystem.h
@@ -44,11 +44,15 @@ struct EDITORFRAMEWORK_API FPlacementOptions
 
 	// If true, asset factory implementations should defer to placing instanced items (i.e. instanced static mesh instead of individual static mesh actors).
 	UPROPERTY()
-	bool bPreferInstancedPlacement;
+	bool bPreferInstancedPlacement = true;
 
 	// If true, asset factory implementations should prefer a batch placement algorithm (like duplicating an object) over a single placement algorithm.
 	UPROPERTY()
-	bool bPreferBatchPlacement;
+	bool bPreferBatchPlacement = true;
+
+	// If true, creates transient preview elements, which are transient and not saved to a level.
+	UPROPERTY()
+	bool bIsCreatingPreviewElements = false;
 };
 
 UCLASS(Transient)
@@ -80,10 +84,15 @@ public:
 	 */
 	TScriptInterface<IAssetFactoryInterface> FindAssetFactoryFromAssetData(const FAssetData& InAssetData);
 
+	// @returns true if the current PlaceAssets call is creating preview elements.
+	bool IsCreatingPreviewElements() const;
+
 private:
 	void RegisterPlacementFactories();
 	void UnregisterPlacementFactories();
 
 	UPROPERTY()
 	TArray<TScriptInterface<IAssetFactoryInterface>> AssetFactories;
+
+	bool bIsCreatingPreviewElements = false;
 };
diff --git a/Engine/Source/Editor/LevelEditor/Private/SLevelEditor.cpp b/Engine/Source/Editor/LevelEditor/Private/SLevelEditor.cpp
--- a/Engine/Source/Editor/LevelEditor/Private/SLevelEditor.cpp
+++ b/Engine/Source/Editor/LevelEditor/Private/SLevelEditor.cpp
@@ -341,6 +341,10 @@ SLevelEditor::~SLevelEditor()
 	// We're going away now, so make sure all toolkits that are hosted within this level editor are shut down
 	FToolkitManager::Get().OnToolkitHostDestroyed( this );
 	HostedToolkits.Reset();
+
+	// Deactivate any active modes, and reset back to the default mode.
+	GetEditorModeManager().SetDefaultMode(FBuiltinEditorModes::EM_Default);
+	GetEditorModeManager().DeactivateAllModes();
 	
 	FLevelEditorModule& LevelEditorModule = FModuleManager::GetModuleChecked< FLevelEditorModule >( LevelEditorModuleName );
 	LevelEditorModule.OnTitleBarMessagesChanged().RemoveAll( this );
diff --git a/Engine/Source/Editor/LevelEditor/Private/SLevelEditorToolBox.cpp b/Engine/Source/Editor/LevelEditor/Private/SLevelEditorToolBox.cpp
--- a/Engine/Source/Editor/LevelEditor/Private/SLevelEditorToolBox.cpp
+++ b/Engine/Source/Editor/LevelEditor/Private/SLevelEditorToolBox.cpp
@@ -273,7 +273,12 @@ void SLevelEditorToolBox::UpdateInlineContent(const TSharedPtr<IToolkit>& Toolki
 			}
 
 			StatusBarMessageHandle = GEditor->GetEditorSubsystem<UStatusBarSubsystem>()->PushStatusBarMessage(LevelEditorStatusBarName, TAttribute<FText>::Create(
-				TAttribute<FText>::FGetter::CreateSP(ModeToolkit.Pin().ToSharedRef(), &FModeToolkit::GetActiveToolDisplayName))
+				TAttribute<FText>::FGetter::CreateLambda([ModeToolkit]()
+					{
+						TSharedPtr<FModeToolkit> PinnedToolkit = ModeToolkit.Pin();
+						if (PinnedToolkit.IsValid()) { return PinnedToolkit->GetActiveToolDisplayName(); }
+						else { return FText(); }
+					}))
 			);
 		}
 	}
diff --git a/Engine/Source/Editor/UnrealEd/Classes/ActorFactories/ActorFactory.h b/Engine/Source/Editor/UnrealEd/Classes/ActorFactories/ActorFactory.h
--- a/Engine/Source/Editor/UnrealEd/Classes/ActorFactories/ActorFactory.h
+++ b/Engine/Source/Editor/UnrealEd/Classes/ActorFactories/ActorFactory.h
@@ -101,4 +101,4 @@ protected:
 	virtual void PostCreateBlueprint( UObject* Asset, AActor* CDO );
 };
 
-extern UNREALED_API FQuat FindActorAlignmentRotation(const FQuat& InActorRotation, const FVector& InModelAxis, const FVector& InWorldNormal);
\ No newline at end of file
+extern UNREALED_API FQuat FindActorAlignmentRotation(const FQuat& InActorRotation, const FVector& InModelAxis, const FVector& InWorldNormal, FQuat* OutDeltaRotation = nullptr);
\ No newline at end of file
diff --git a/Engine/Source/Editor/UnrealEd/Private/EditorModeManager.cpp b/Engine/Source/Editor/UnrealEd/Private/EditorModeManager.cpp
--- a/Engine/Source/Editor/UnrealEd/Private/EditorModeManager.cpp
+++ b/Engine/Source/Editor/UnrealEd/Private/EditorModeManager.cpp
@@ -642,7 +642,7 @@ void FEditorModeTools::RebuildModeToolBar()
 					}
 					else if (UEdMode* ScriptableMode = GetActiveScriptableMode(Row.ModeID))
 					{
-						RowToolkit = ScriptableMode->GetToolkit();
+						RowToolkit = ScriptableMode->GetToolkit().Pin();
 					}
 
 					bExclusivePalettes = RowToolkit->HasExclusiveToolPalettes();
@@ -777,7 +777,7 @@ void FEditorModeTools::InvokeToolPaletteTab(FEditorModeID InModeID, FName InPale
 			}
 			else if (UEdMode* ScriptableMode = GetActiveScriptableMode(InModeID))
 			{
-				RowToolkit = ScriptableMode->GetToolkit();
+				RowToolkit = ScriptableMode->GetToolkit().Pin();
 			}
 
 			TSharedPtr<SWidget> ActiveWidget = ModeToolbarPaletteSwitcher.Pin()->GetActiveWidget();
@@ -943,11 +943,7 @@ void FEditorModeTools::ActivateMode(FEditorModeID InID, bool bToggle)
 
 	// Recycle a mode or factory a new one
 	UEdMode* ScriptableMode = RecycledScriptableModes.FindRef(InID);
-	if (ScriptableMode)
-	{
-		RecycledScriptableModes.Remove(InID);
-	}
-	else
+	if (!ScriptableMode)
 	{
 		ScriptableMode = GEditor->GetEditorSubsystem<UAssetEditorSubsystem>()->CreateEditorModeWithToolsOwner(InID, *this);
 	}
@@ -979,7 +975,7 @@ void FEditorModeTools::ActivateMode(FEditorModeID InID, bool bToggle)
 
 	// Ask the mode to build the toolbar.
 	TSharedPtr<FUICommandList> CommandList;
-	const TSharedPtr<FModeToolkit> Toolkit = ScriptableMode->GetToolkit();
+	const TSharedPtr<FModeToolkit> Toolkit = ScriptableMode->GetToolkit().Pin();
 	if (Toolkit.IsValid())
 	{
 		CommandList = Toolkit->GetToolkitCommands();
@@ -1003,6 +999,9 @@ void FEditorModeTools::ActivateMode(FEditorModeID InID, bool bToggle)
 			SpawnOrUpdateModeToolbar();
 		}
 	}
+
+	RecycledScriptableModes.Remove(InID);
+
 	// Update the editor UI
 	FEditorSupportDelegates::UpdateUI.Broadcast();
 }
diff --git a/Engine/Source/Editor/UnrealEd/Private/Factories/ActorFactory.cpp b/Engine/Source/Editor/UnrealEd/Private/Factories/ActorFactory.cpp
--- a/Engine/Source/Editor/UnrealEd/Private/Factories/ActorFactory.cpp
+++ b/Engine/Source/Editor/UnrealEd/Private/Factories/ActorFactory.cpp
@@ -149,7 +149,7 @@ DEFINE_LOG_CATEGORY(LogActorFactory);
  * This function attempts to find a 'natural' looking rotation by rotating around a local pitch axis, and a world Z. Rotating in this way
  * should retain the roll around the model space axis, removing rotation artifacts introduced by a simpler quaternion rotation.
  */
-FQuat FindActorAlignmentRotation(const FQuat& InActorRotation, const FVector& InModelAxis, const FVector& InWorldNormal)
+FQuat FindActorAlignmentRotation(const FQuat& InActorRotation, const FVector& InModelAxis, const FVector& InWorldNormal, FQuat* OutDeltaRotation)
 {
 	FVector TransformedModelAxis = InActorRotation.RotateVector(InModelAxis);
 
@@ -157,6 +157,10 @@ FQuat FindActorAlignmentRotation(const FQuat& InActorRotation, const FVector& In
 	const auto DestNormalModelSpace = InverseActorRotation.RotateVector(InWorldNormal);
 
 	FQuat DeltaRotation = FQuat::Identity;
+	if (OutDeltaRotation)
+	{
+		*OutDeltaRotation = FQuat::Identity;
+	}
 
 	const float VectorDot = InWorldNormal | TransformedModelAxis;
 	if (1.f - FMath::Abs(VectorDot) <= KINDA_SMALL_NUMBER)
@@ -197,6 +201,11 @@ FQuat FindActorAlignmentRotation(const FQuat& InActorRotation, const FVector& In
 			const FVector WorldYawAxis = (InActorRotation * DeltaRotation).Inverse().RotateVector(Z);
 			DeltaRotation *= FQuat(WorldYawAxis, -Yaw);
 		}
+
+		if (OutDeltaRotation)
+		{
+			*OutDeltaRotation = DeltaRotation;
+		}
 	}
 
 	return InActorRotation * DeltaRotation;
@@ -296,7 +305,14 @@ TArray<FTypedElementHandle> UActorFactory::PlaceAsset(const FAssetPlacementInfo&
 	if (NewActor)
 	{
 		PlacedActorHandles.Add(UEngineElementsLibrary::AcquireEditorActorElementHandle(NewActor));
+
+		// Disable collision for preview actors
+		if (InPlacementOptions.bIsCreatingPreviewElements)
+		{
+			NewActor->SetActorEnableCollision(false);
+		}
 	}
+
 	return PlacedActorHandles;
 }
 
@@ -410,9 +426,16 @@ AActor* UActorFactory::SpawnActor( UObject* Asset, ULevel* InLevel, const FTrans
 		SpawnInfo.OverrideLevel = InLevel;
 		SpawnInfo.ObjectFlags = InObjectFlags;
 		SpawnInfo.Name = Name;
-#if WITH_EDITOR
-		SpawnInfo.bTemporaryEditorActor = FLevelEditorViewportClient::IsDroppingPreviewActor();
-#endif
+
+		const bool bIsCreatingPreviewElements = FLevelEditorViewportClient::IsDroppingPreviewActor();
+		bool bIsPlacementSystemCreatingPreviewElements = false;
+		if (UPlacementSubsystem* PlacementSubsystem = GEditor->GetEditorSubsystem<UPlacementSubsystem>())
+		{
+			bIsPlacementSystemCreatingPreviewElements = PlacementSubsystem->IsCreatingPreviewElements();
+		}
+		SpawnInfo.bTemporaryEditorActor = bIsCreatingPreviewElements || bIsPlacementSystemCreatingPreviewElements;
+		SpawnInfo.bHideFromSceneOutliner = bIsPlacementSystemCreatingPreviewElements;
+
 		return InLevel->OwningWorld->SpawnActor( DefaultActor->GetClass(), &Transform, SpawnInfo );
 	}
 
diff --git a/Engine/Source/Editor/UnrealEd/Private/Factories/EditorStaticMeshFactory.cpp b/Engine/Source/Editor/UnrealEd/Private/Factories/EditorStaticMeshFactory.cpp
--- a/Engine/Source/Editor/UnrealEd/Private/Factories/EditorStaticMeshFactory.cpp
+++ b/Engine/Source/Editor/UnrealEd/Private/Factories/EditorStaticMeshFactory.cpp
@@ -19,7 +19,8 @@
 
 TArray<FTypedElementHandle> UEditorStaticMeshFactory::PlaceAsset(const FAssetPlacementInfo& InPlacementInfo, const FPlacementOptions& InPlacementOptions)
 {
-	if (!InPlacementOptions.bPreferInstancedPlacement)
+	// If we're disallowing instanced placement, or creating preview elements, don't use the ISM placement.
+	if (!InPlacementOptions.bPreferInstancedPlacement || InPlacementOptions.bIsCreatingPreviewElements)
 	{
 		return Super::PlaceAsset(InPlacementInfo, InPlacementOptions);
 	}
diff --git a/Engine/Source/Editor/UnrealEd/Public/Tools/UEdMode.h b/Engine/Source/Editor/UnrealEd/Public/Tools/UEdMode.h
--- a/Engine/Source/Editor/UnrealEd/Public/Tools/UEdMode.h
+++ b/Engine/Source/Editor/UnrealEd/Public/Tools/UEdMode.h
@@ -221,7 +221,7 @@ public:
 	virtual bool UsesToolkits() const;
 
 	/** Gets the toolkit created by this mode */
-	TSharedPtr<FModeToolkit> GetToolkit() { return Toolkit; }
+	TWeakPtr<FModeToolkit> GetToolkit() { return Toolkit; }
 
 	/** Returns the world this toolkit is editing */
 	UWorld* GetWorld() const;
