Placement Mode Single place tool improvements:
* Tweak after placement in world and local space using standard level editor gizmos.
* Removed Single place tool settings. Since we're using the TRS gizmos, these settings are unnecessary.
* Moved cycle through palette to shift + right click to match other viewport functions (i.e. cycling static mesh previews).
* Placement mode now inherits from default ed mode, so we will get the new gizmos as they come online.
* Fixed a bug in the select tool where the input was never making it back to the viewport for handling.
* Fixed a bug in the ITF Editor hooks where the mouse capture was incorrectly reporting and eating the input from the viewport.

diff --git a/Engine/Source/Editor/AssetPlacementEdMode/Private/AssetPlacementEdMode.cpp b/Engine/Source/Editor/AssetPlacementEdMode/Private/AssetPlacementEdMode.cpp
--- a/Engine/Source/Editor/AssetPlacementEdMode/Private/AssetPlacementEdMode.cpp
+++ b/Engine/Source/Editor/AssetPlacementEdMode/Private/AssetPlacementEdMode.cpp
@@ -42,8 +42,6 @@ UAssetPlacementEdMode::UAssetPlacementEdMode()
 		LOCTEXT("AssetPlacementEdModeName", "Placement"),
 		FSlateIcon(FAssetPlacementEdModeStyle::GetStyleSetName(), "PlacementBrowser.ShowAllContent"),
 		MoveTemp(IsEnabledAttr));
-
-	bIsInSelectionTool = false;
 }
 
 UAssetPlacementEdMode::~UAssetPlacementEdMode()
@@ -52,8 +50,8 @@ UAssetPlacementEdMode::~UAssetPlacementEdMode()
 
 void UAssetPlacementEdMode::Enter()
 {
-	// Set the settings object before we call Super::Enter, so that it's available for hooking up in the toolkit
-	SettingsObject = const_cast<UAssetPlacementSettings*>(GEditor->GetEditorSubsystem<UPlacementModeSubsystem>()->GetModeSettingsObject().Get());
+	// Set the settings object before we call Super::Enter, since we're using a shared one from the subsystem.
+	SettingsObject = const_cast<UAssetPlacementSettings*>(GEditor->GetEditorSubsystem<UPlacementModeSubsystem>()->GetModeSettingsObject());
 
 	Super::Enter();
 
@@ -64,14 +62,15 @@ void UAssetPlacementEdMode::Enter()
 	RegisterTool(PlacementModeCommands.PlaceSingle, UPlacementModePlaceSingleTool::ToolName, NewObject<UPlacementModePlaceSingleToolBuilder>(this));
 	RegisterTool(PlacementModeCommands.Erase, UPlacementModeEraseTool::ToolName, NewObject<UPlacementModeEraseToolBuilder>(this));
 
+	// Stash the current editor selection, since this mode will modify it.
+	Owner->StoreSelection(AssetPlacementEdModeID);
+	bIsInSelectionTool = false;
+
 	// Enable the select tool by default.
 	// Disable undo tracking so that we can't accidentally undo ourselves out of the select mode and into an invalid state.
 	GetToolManager()->ConfigureChangeTrackingMode(EToolChangeTrackingMode::NoChangeTracking);
 	ToolsContext->StartTool(UPlacementModeSelectTool::ToolName);
 	GetToolManager()->ConfigureChangeTrackingMode(EToolChangeTrackingMode::FullUndoRedo);
-
-	// Stash the current editor selection, since this mode will modify it.
-	Owner->StoreSelection(AssetPlacementEdModeID);
 }
 
 void UAssetPlacementEdMode::Exit()
@@ -85,18 +84,24 @@ void UAssetPlacementEdMode::Exit()
 		IFA->SelectInstance(nullptr, 0, bAppendSelection);
 	}
 
-	// Restore the selection to the original state.
+	Super::Exit();
+
+	// Restore the selection to the original state after all the tools have shutdown in UEdMode::Exit()
+	// Since they can continue messing with selection states.
 	Owner->RestoreSelection(AssetPlacementEdModeID);
 
 	SettingsObjectAsPlacementSettings.Reset();
-
-	Super::Exit();
 }
 
 void UAssetPlacementEdMode::CreateToolkit()
 {
 	SettingsObjectAsPlacementSettings = Cast<UAssetPlacementSettings>(SettingsObject);
-	Toolkit = MakeShareable(new FAssetPlacementEdModeToolkit(SettingsObjectAsPlacementSettings));
+	Toolkit = MakeShared<FAssetPlacementEdModeToolkit>();
+}
+
+bool UAssetPlacementEdMode::UsesToolkits() const
+{
+	return true;
 }
 
 TMap<FName, TArray<TSharedPtr<FUICommandInfo>>> UAssetPlacementEdMode::GetModeCommands() const
@@ -143,7 +148,8 @@ void UAssetPlacementEdMode::OnToolStarted(UInteractiveToolManager* Manager, UInt
 
 	bool bWasInSelectionTool = bIsInSelectionTool;
 	FString ActiveToolName = GetToolManager()->GetActiveToolName(EToolSide::Mouse);
-	if (ActiveToolName == UPlacementModeSelectTool::ToolName || ActiveToolName == UPlacementModeLassoSelectTool::ToolName)
+	bool bIsSinglePlaceTool = ActiveToolName == UPlacementModePlaceSingleTool::ToolName;
+	if (ActiveToolName == UPlacementModeSelectTool::ToolName || ActiveToolName == UPlacementModeLassoSelectTool::ToolName || bIsSinglePlaceTool)
 	{
 		bIsInSelectionTool = true;
 	}
@@ -152,17 +158,12 @@ void UAssetPlacementEdMode::OnToolStarted(UInteractiveToolManager* Manager, UInt
 		bIsInSelectionTool = false;
 	}
 
-	// Stash or pop selection
-	if (bWasInSelectionTool != bIsInSelectionTool)
+	// Restore the selection if we're going into the selection tools.
+	// Allow the selection to be empty if we're going into single place tool for a clean slate.
+	bool bRestoreSelectionState = bIsInSelectionTool && !bWasInSelectionTool && !bIsSinglePlaceTool;
+	if (bRestoreSelectionState)
 	{
-		if (bIsInSelectionTool)
-		{
-			Owner->RestoreSelection(UPlacementModeSelectTool::ToolName);
-		}
-		else
-		{
-			Owner->StoreSelection(UPlacementModeSelectTool::ToolName);
-		}
+		Owner->RestoreSelection(UPlacementModeSelectTool::ToolName);
 
 		// Todo: remove with HISM element handles
 		for (TActorIterator<AInstancedFoliageActor> It(GetWorld()); It; ++It)
@@ -171,6 +172,23 @@ void UAssetPlacementEdMode::OnToolStarted(UInteractiveToolManager* Manager, UInt
 			IFA->ApplySelection(bIsInSelectionTool);
 		}
 	}
+	else if (!bIsInSelectionTool)
+	{
+		// If we can't select, clear out the selection set for the active tool.
+		ClearSelection();
+	}
+}
+
+void UAssetPlacementEdMode::OnToolEnded(UInteractiveToolManager* Manager, UInteractiveTool* Tool)
+{
+	Super::OnToolEnded(Manager, Tool);
+
+	// Always store the most recent selection, even if we are leaving single placement tool to preserve what the user was doing last.
+	if (bIsInSelectionTool)
+	{
+		constexpr bool bClearSelection = false;
+		Owner->StoreSelection(UPlacementModeSelectTool::ToolName, false);
+	}
 }
 
 bool UAssetPlacementEdMode::UsesPropertyWidgets() const
diff --git a/Engine/Source/Editor/AssetPlacementEdMode/Private/AssetPlacementEdModeToolkit.cpp b/Engine/Source/Editor/AssetPlacementEdMode/Private/AssetPlacementEdModeToolkit.cpp
--- a/Engine/Source/Editor/AssetPlacementEdMode/Private/AssetPlacementEdModeToolkit.cpp
+++ b/Engine/Source/Editor/AssetPlacementEdMode/Private/AssetPlacementEdModeToolkit.cpp
@@ -1,21 +1,14 @@
 // Copyright Epic Games, Inc. All Rights Reserved.
 
 #include "AssetPlacementEdModeToolkit.h"
-#include "AssetPlacementEdMode.h"
-#include "AssetPlacementSettings.h"
-#include "SAssetPlacementPalette.h"
-
-#include "Modules/ModuleManager.h"
 
+#include "SAssetPlacementPalette.h"
 #include "Widgets/Layout/SExpandableArea.h"
 #include "Widgets/Layout/SScrollBox.h"
 
-#include "EditorModeManager.h"
-
 #define LOCTEXT_NAMESPACE "AssetPlacementEdModeToolkit"
 
-FAssetPlacementEdModeToolkit::FAssetPlacementEdModeToolkit(TWeakObjectPtr<UAssetPlacementSettings> InPlacementSettings)
-	: PlacementSettings(InPlacementSettings)
+FAssetPlacementEdModeToolkit::FAssetPlacementEdModeToolkit()
 {
 }
 
@@ -62,7 +55,6 @@ TSharedPtr<SWidget> FAssetPlacementEdModeToolkit::GetInlineContent() const
 				+ SScrollBox::Slot()
 				[
 					SNew(SAssetPlacementPalette)
-					.PlacementSettings(PlacementSettings)
 				]
 			]
 		];
diff --git a/Engine/Source/Editor/AssetPlacementEdMode/Private/Modes/PlacementModeSubsystem.cpp b/Engine/Source/Editor/AssetPlacementEdMode/Private/Modes/PlacementModeSubsystem.cpp
--- a/Engine/Source/Editor/AssetPlacementEdMode/Private/Modes/PlacementModeSubsystem.cpp
+++ b/Engine/Source/Editor/AssetPlacementEdMode/Private/Modes/PlacementModeSubsystem.cpp
@@ -20,7 +20,7 @@ void UPlacementModeSubsystem::Deinitialize()
 	ModeSettings = nullptr;
 }
 
-TWeakObjectPtr<const UAssetPlacementSettings> UPlacementModeSubsystem::GetModeSettingsObject() const
+const UAssetPlacementSettings* UPlacementModeSubsystem::GetModeSettingsObject() const
 {
 	return ModeSettings;
 }
@@ -57,3 +57,14 @@ bool UPlacementModeSubsystem::DoesCurrentPaletteSupportElement(const FTypedEleme
 
 	return false;
 }
+
+bool UPlacementModeSubsystem::AddPaletteItem(const FPaletteItem& InPaletteItem)
+{
+	if (ModeSettings && !ModeSettings->PaletteItems.FindByPredicate([InPaletteItem](const FPaletteItem& ItemIter) { return ItemIter.AssetData.ObjectPath == InPaletteItem.AssetData.ObjectPath; }))
+	{
+		ModeSettings->PaletteItems.Add(InPaletteItem);
+		return true;
+	}
+
+	return false;
+}
diff --git a/Engine/Source/Editor/AssetPlacementEdMode/Private/Modes/PlacementModeSubsystem.h b/Engine/Source/Editor/AssetPlacementEdMode/Private/Modes/PlacementModeSubsystem.h
--- a/Engine/Source/Editor/AssetPlacementEdMode/Private/Modes/PlacementModeSubsystem.h
+++ b/Engine/Source/Editor/AssetPlacementEdMode/Private/Modes/PlacementModeSubsystem.h
@@ -7,6 +7,7 @@
 
 #include "PlacementModeSubsystem.generated.h"
 
+struct FPaletteItem;
 struct FTypedElementHandle;
 class UAssetPlacementSettings;
 
@@ -22,7 +23,7 @@ public:
 	// End USubsystem Interface
 
 	// @returns the settings object for the mode for sharing across all tools and tool builders.
-	TWeakObjectPtr<const UAssetPlacementSettings> GetModeSettingsObject() const;
+	const UAssetPlacementSettings* GetModeSettingsObject() const;
 
 	/**
 	 * Verifies if the given element handle is supported by the current mode settings' palette.
@@ -31,6 +32,13 @@ public:
 	 */
 	bool DoesCurrentPaletteSupportElement(const FTypedElementHandle& InElementToCheck) const;
 
+	/**
+	 * Adds the given palette item to the current palette.
+	 * 
+	 * @returns true if the item was added.
+	 */
+	bool AddPaletteItem(const FPaletteItem& InPaletteItem);
+
 protected:
 	UPROPERTY()
 	TObjectPtr<UAssetPlacementSettings> ModeSettings;
diff --git a/Engine/Source/Editor/AssetPlacementEdMode/Private/SAssetPlacementPalette.cpp b/Engine/Source/Editor/AssetPlacementEdMode/Private/SAssetPlacementPalette.cpp
--- a/Engine/Source/Editor/AssetPlacementEdMode/Private/SAssetPlacementPalette.cpp
+++ b/Engine/Source/Editor/AssetPlacementEdMode/Private/SAssetPlacementPalette.cpp
@@ -38,6 +38,7 @@
 #include "Subsystems/PlacementSubsystem.h"
 #include "Engine/AssetManager.h"
 #include "Engine/StreamableManager.h"
+#include "Modes/PlacementModeSubsystem.h"
 
 #define LOCTEXT_NAMESPACE "AssetPlacementMode"
 
@@ -111,11 +112,6 @@ void SAssetPlacementPalette::Construct(const FArguments& InArgs)
 	bItemsNeedRebuild = false;
 	bIsRebuildTimerRegistered = false;
 	bIsRefreshTimerRegistered = false;
-	PlacementSettings = InArgs._PlacementSettings;
-	if (!PlacementSettings.IsValid())
-	{
-		PlacementSettings = GetMutableDefault<UAssetPlacementSettings>();
-	}
 
 	UICommandList = MakeShared<FUICommandList>();
 
@@ -227,6 +223,15 @@ void SAssetPlacementPalette::Construct(const FArguments& InArgs)
 		]
 	];
 
+	TWeakObjectPtr<const UAssetPlacementSettings> PlacementSettings = GEditor->GetEditorSubsystem<UPlacementModeSubsystem>()->GetModeSettingsObject();
+	if (PlacementSettings.IsValid())
+	{
+		for (const FPaletteItem& PaletteItem : PlacementSettings->PaletteItems)
+		{
+			PaletteItems.Add(MakeShared<FAssetPlacementPaletteItemModel>(MakeShared<FPaletteItem>(PaletteItem), SharedThis(this), ThumbnailPool));
+		}
+	}
+
 	UpdatePalette(true);
 }
 
@@ -305,8 +310,12 @@ void SAssetPlacementPalette::AddPlacementType(const FAssetData& AssetData)
 		return;
 	}
 
-	// Check that we don't already have this item in the palette
-	if (PaletteItems.FindByPredicate([&AssetData](TSharedPtr<FAssetPlacementPaletteItemModel>& InItem) { return InItem.IsValid() && (InItem->GetTypeUIInfo()->AssetData.ObjectPath == AssetData.ObjectPath); }))
+	FAssetPlacementUIInfoPtr PlacementInfo = MakeShared<FPaletteItem>();
+	PlacementInfo->AssetData = AssetData;
+	PlacementInfo->FactoryOverride = FactoryInterface;
+
+	// Try to add the item to the mode's palette
+	if (!GEditor->GetEditorSubsystem<UPlacementModeSubsystem>()->AddPaletteItem(*PlacementInfo))
 	{
 		return;
 	}
@@ -314,14 +323,7 @@ void SAssetPlacementPalette::AddPlacementType(const FAssetData& AssetData)
 	// Try to load the asset async so it's ready to place.
 	UAssetManager::GetStreamableManager().RequestAsyncLoad(AssetData.ToSoftObjectPath());
 
-	FAssetPlacementUIInfoPtr PlacementInfo = MakeShared<FPaletteItem>();
-	PlacementInfo->AssetData = AssetData;
-	PlacementInfo->FactoryOverride = FactoryInterface;
 	PaletteItems.Add(MakeShared<FAssetPlacementPaletteItemModel>(PlacementInfo, SharedThis(this), ThumbnailPool));
-	if (PlacementSettings.IsValid())
-	{
-		PlacementSettings->PaletteItems.Add(*PlacementInfo);
-	}
 	UpdatePalette(true);
 }
 
diff --git a/Engine/Source/Editor/AssetPlacementEdMode/Private/SAssetPlacementPalette.h b/Engine/Source/Editor/AssetPlacementEdMode/Private/SAssetPlacementPalette.h
--- a/Engine/Source/Editor/AssetPlacementEdMode/Private/SAssetPlacementPalette.h
+++ b/Engine/Source/Editor/AssetPlacementEdMode/Private/SAssetPlacementPalette.h
@@ -43,7 +43,6 @@ class SAssetPlacementPalette : public SCompoundWidget
 {
 public:
 	SLATE_BEGIN_ARGS(SAssetPlacementPalette) {}
-		SLATE_ARGUMENT(TWeakObjectPtr<UAssetPlacementSettings>, PlacementSettings)
 		SLATE_ARGUMENT(TSharedPtr<IPropertyHandle>, PalettePropertyHandle)
 	SLATE_END_ARGS()
 
@@ -196,7 +195,6 @@ private:
 	/** Placement mesh details widget  */
 	TSharedPtr<class IDetailsView> DetailsWidget;
 
-	TWeakObjectPtr<UAssetPlacementSettings> PlacementSettings;
 	TSharedPtr<IPropertyHandle> PalettePropertyHandle;
 
 	/** Placement items search box widget */
diff --git a/Engine/Source/Editor/AssetPlacementEdMode/Private/Tools/PlacementClickDragToolBase.cpp b/Engine/Source/Editor/AssetPlacementEdMode/Private/Tools/PlacementClickDragToolBase.cpp
--- a/Engine/Source/Editor/AssetPlacementEdMode/Private/Tools/PlacementClickDragToolBase.cpp
+++ b/Engine/Source/Editor/AssetPlacementEdMode/Private/Tools/PlacementClickDragToolBase.cpp
@@ -11,12 +11,18 @@ void UPlacementClickDragToolBase::Setup()
 	Super::Setup();
 
 	RemoveToolPropertySource(BrushProperties);
-	BrushStampIndicator->bDrawRadiusCircle = false;
-	BrushProperties->BrushSize = .5f;
-	BrushProperties->BrushFalloffAmount = 1.f;
 }
 
 double UPlacementClickDragToolBase::EstimateMaximumTargetDimension()
 {
 	return Super::EstimateMaximumTargetDimension();
-}
\ No newline at end of file
+}
+
+void UPlacementClickDragToolBase::SetupBrushStampIndicator()
+{
+	Super::SetupBrushStampIndicator();
+
+	BrushStampIndicator->bDrawRadiusCircle = false;
+	BrushProperties->BrushSize = .5f;
+	BrushProperties->BrushFalloffAmount = 1.f;
+}
diff --git a/Engine/Source/Editor/AssetPlacementEdMode/Private/Tools/PlacementClickDragToolBase.h b/Engine/Source/Editor/AssetPlacementEdMode/Private/Tools/PlacementClickDragToolBase.h
--- a/Engine/Source/Editor/AssetPlacementEdMode/Private/Tools/PlacementClickDragToolBase.h
+++ b/Engine/Source/Editor/AssetPlacementEdMode/Private/Tools/PlacementClickDragToolBase.h
@@ -15,4 +15,5 @@ public:
 
 protected:
 	virtual double EstimateMaximumTargetDimension() override;
+	virtual void SetupBrushStampIndicator() override;
 };
diff --git a/Engine/Source/Editor/AssetPlacementEdMode/Private/Tools/PlacementPlaceSingleTool.cpp b/Engine/Source/Editor/AssetPlacementEdMode/Private/Tools/PlacementPlaceSingleTool.cpp
--- a/Engine/Source/Editor/AssetPlacementEdMode/Private/Tools/PlacementPlaceSingleTool.cpp
+++ b/Engine/Source/Editor/AssetPlacementEdMode/Private/Tools/PlacementPlaceSingleTool.cpp
@@ -1,18 +1,20 @@
 // Copyright Epic Games, Inc. All Rights Reserved.
 
 #include "Tools/PlacementPlaceSingleTool.h"
-#include "UObject/Object.h"
+
 #include "AssetPlacementSettings.h"
-#include "Subsystems/PlacementSubsystem.h"
-#include "ScopedTransaction.h"
-#include "Editor/EditorEngine.h"
 #include "Editor.h"
-#include "Modes/PlacementModeSubsystem.h"
 #include "InteractiveToolManager.h"
-#include "Elements/Framework/TypedElementRegistry.h"
+#include "ScopedTransaction.h"
+#include "BaseBehaviors/ClickDragBehavior.h"
+#include "Editor/EditorEngine.h"
 #include "Elements/Framework/TypedElementHandle.h"
-#include "Elements/Interfaces/TypedElementObjectInterface.h"
+#include "Elements/Framework/TypedElementRegistry.h"
+#include "Elements/Framework/TypedElementSelectionSet.h"
 #include "Elements/Interfaces/TypedElementWorldInterface.h"
+#include "Modes/PlacementModeSubsystem.h"
+#include "Subsystems/PlacementSubsystem.h"
+#include "UObject/Object.h"
 
 constexpr TCHAR UPlacementModePlaceSingleTool::ToolName[];
 
@@ -21,22 +23,6 @@ UPlacementBrushToolBase* UPlacementModePlaceSingleToolBuilder::FactoryToolInstan
 	return NewObject<UPlacementModePlaceSingleTool>(Outer);
 }
 
-bool UPlacementModePlaceSingleSettings::CanEditChange(const FProperty* Property) const
-{
-	if (!Super::CanEditChange(Property))
-	{
-		return false;
-	}
-
-	const FName PropertyName = Property->GetFName();
-	if (PropertyName == GET_MEMBER_NAME_CHECKED(UPlacementModePlaceSingleSettings, bInvertCursorAxis))
-	{
-		return bAlignToCursor;
-	}
-
-	return true;
-}
-
 UPlacementModePlaceSingleTool::UPlacementModePlaceSingleTool() = default;
 UPlacementModePlaceSingleTool::~UPlacementModePlaceSingleTool() = default;
 UPlacementModePlaceSingleTool::UPlacementModePlaceSingleTool(FVTableHelper& Helper)
@@ -48,18 +34,18 @@ void UPlacementModePlaceSingleTool::Setup()
 {
 	Super::Setup();
 
-	SingleToolSettings = NewObject<UPlacementModePlaceSingleSettings>(this);
-	SingleToolSettings->LoadConfig();
-	AddToolPropertySource(SingleToolSettings);
+	SetupRightClickMouseBehavior();
+
+	bIsTweaking = false;
 }
 
 void UPlacementModePlaceSingleTool::Shutdown(EToolShutdownType ShutdownType)
 {
 	DestroyPreviewElements();
 
-	SingleToolSettings->SaveConfig();
-	RemoveToolPropertySource(SingleToolSettings);
-	SingleToolSettings = nullptr;
+	// Preserve the selection on exiting the tool, so that a users' state persists as they use the mode.
+	constexpr bool bClearSelectionSet = false;
+	ExitTweakState(bClearSelectionSet);
 
 	Super::Shutdown(ShutdownType);
 }
@@ -68,30 +54,23 @@ void UPlacementModePlaceSingleTool::OnClickPress(const FInputDeviceRay& PressPos
 {
 	Super::OnClickPress(PressPos);
 
-	if (bCtrlToggle)
-	{
-		// If this was a ctrl click, regenerate the preview placement information, and destroy any active gizmos.
-		GeneratePreviewPlacementData(PressPos);
-	}
-	else
-	{
-		// Otherwise, and we're not currently editing something that was placed, then place the preview data.
-		// Todo - verify we don't get here if the gizmos are active
-		FPlacementOptions PlacementOptions;
-		PlacementOptions.bPreferBatchPlacement = true;
-		PlacementOptions.bPreferInstancedPlacement = true;
+	// Update the preview element one final time to be sure the transform is updated.
+	UpdatePreviewElements(PressPos);
+	DestroyPreviewElements();
 
-		UPlacementSubsystem* PlacementSubsystem = GEditor->GetEditorSubsystem<UPlacementSubsystem>();
-		if (PlacementSubsystem)
-		{
-			// Update the level, just in case the preview moved us out of the current one.
-			PreviewPlacementInfo->PreferredLevel = GEditor->GetEditorWorldContext().World()->GetCurrentLevel();
+	// Place the Preview data if we managed to get to a valid handled click.
+	FPlacementOptions PlacementOptions;
+	PlacementOptions.bPreferBatchPlacement = true;
+	PlacementOptions.bPreferInstancedPlacement = true;
 
-			FScopedTransaction Transaction(NSLOCTEXT("PlacementMode", "SinglePlaceAsset", "Place Single Asset"));
-			PlacementSubsystem->PlaceAsset(*PreviewPlacementInfo, PlacementOptions);
-		}
+	UPlacementSubsystem* PlacementSubsystem = GEditor->GetEditorSubsystem<UPlacementSubsystem>();
+	if (PlacementSubsystem)
+	{
+		// Update the level, just in case the preview moved us out of the current one.
+		PreviewPlacementInfo->PreferredLevel = GEditor->GetEditorWorldContext().World()->GetCurrentLevel();
 
-		DestroyPreviewElements();
+		FScopedTransaction Transaction(NSLOCTEXT("PlacementMode", "SinglePlaceAsset", "Place Single Asset"));
+		EnterTweakState(PlacementSubsystem->PlaceAsset(*PreviewPlacementInfo, PlacementOptions));
 	}
 }
 
@@ -100,7 +79,7 @@ void UPlacementModePlaceSingleTool::OnBeginHover(const FInputDeviceRay& DevicePo
 	Super::OnBeginHover(DevicePos);
 
 	// Always regenerate the placement data when a hover sequence begins
-	GeneratePreviewPlacementData(DevicePos);
+	CreatePreviewElements(DevicePos);
 }
 
 bool UPlacementModePlaceSingleTool::OnUpdateHover(const FInputDeviceRay& DevicePos)
@@ -119,16 +98,36 @@ void UPlacementModePlaceSingleTool::OnEndHover()
 {
 	Super::OnEndHover();
 
-	// remove the preview elements from the level
+	// Destroy the preview elements we created.
 	DestroyPreviewElements();
 }
 
+FInputRayHit UPlacementModePlaceSingleTool::CanBeginClickDragSequence(const FInputDeviceRay& PressPos)
+{
+	if (!bIsTweaking)
+	{
+		return Super::CanBeginClickDragSequence(PressPos);
+	}
+	
+	return FInputRayHit();
+}
+
+FInputRayHit UPlacementModePlaceSingleTool::BeginHoverSequenceHitTest(const FInputDeviceRay& DevicePos)
+{
+	if (!bIsTweaking)
+	{
+		return Super::BeginHoverSequenceHitTest(DevicePos);
+	}
+
+	return FInputRayHit();
+}
+
 void UPlacementModePlaceSingleTool::GeneratePreviewPlacementData(const FInputDeviceRay& DevicePos)
 {
 	PreviewPlacementInfo.Reset();
 
-	TWeakObjectPtr<const UAssetPlacementSettings> PlacementSettings = GEditor->GetEditorSubsystem<UPlacementModeSubsystem>()->GetModeSettingsObject();
-	if (PlacementSettings.IsValid() && PlacementSettings->PaletteItems.Num())
+	const UAssetPlacementSettings* PlacementSettings = GEditor->GetEditorSubsystem<UPlacementModeSubsystem>()->GetModeSettingsObject();
+	if (PlacementSettings && PlacementSettings->PaletteItems.Num())
 	{
 		int32 ItemIndex = FMath::RandHelper(PlacementSettings->PaletteItems.Num());
 		const FPaletteItem& ItemToPlace = PlacementSettings->PaletteItems[ItemIndex];
@@ -144,7 +143,7 @@ void UPlacementModePlaceSingleTool::GeneratePreviewPlacementData(const FInputDev
 void UPlacementModePlaceSingleTool::CreatePreviewElements(const FInputDeviceRay& DevicePos)
 {
 	// Place the preview elements from our stored info
-	if (!PreviewPlacementInfo.IsValid())
+	if (!PreviewPlacementInfo)
 	{
 		GeneratePreviewPlacementData(DevicePos);
 	}
@@ -156,24 +155,21 @@ void UPlacementModePlaceSingleTool::CreatePreviewElements(const FInputDeviceRay&
 
 		PreviewElements = PlacementSubsystem->PlaceAsset(*PreviewPlacementInfo, PlacementOptions);
 	}
-	
-	// Disable collision on any preview elements
-	for (FTypedElementHandle& PreviewElement : PreviewElements)
-	{
-		if (TTypedElement<UTypedElementWorldInterface> WorldInterfaceElement = UTypedElementRegistry::GetInstance()->GetElement<UTypedElementWorldInterface>(PreviewElement))
-		{
-			WorldInterfaceElement.NotifyMovementStarted();
-		}
-	}
+
+	SetupBrushStampIndicator();
+
+	NotifyMovementStarted(PreviewElements);
 }
 
 void UPlacementModePlaceSingleTool::UpdatePreviewElements(const FInputDeviceRay& DevicePos)
 {
-	if (PreviewElements.Num() == 0)
+	// If we should have preview elements, but do not currently, go ahead and create them.
+	if ((PreviewElements.Num() == 0) && !BrushStampIndicator)
 	{
 		CreatePreviewElements(DevicePos);
 	}
 
+	// If we don't actually have any preview handles created, we don't need to update them, so go ahead and bail.
 	if (PreviewElements.Num() == 0)
 	{
 		return;
@@ -184,31 +180,80 @@ void UPlacementModePlaceSingleTool::UpdatePreviewElements(const FInputDeviceRay&
 	FinalizedTransform.SetLocation(LastBrushStamp.HitResult.ImpactPoint);
 
 	// Update the rotation.
-	TWeakObjectPtr<const UAssetPlacementSettings> PlacementSettings = GEditor->GetEditorSubsystem<UPlacementModeSubsystem>()->GetModeSettingsObject();
-	if (PlacementSettings.IsValid() && PlacementSettings->bAlignToNormal)
+	const UAssetPlacementSettings* PlacementSettings = GEditor->GetEditorSubsystem<UPlacementModeSubsystem>()->GetModeSettingsObject();
+	if (PlacementSettings && PlacementSettings->bAlignToNormal)
 	{
 		FQuat UpdatedRotation = UpdateRotationAlignedToBrushNormal(PlacementSettings->AxisToAlignWithNormal, PlacementSettings->bInvertNormalAxis);
 		FinalizedTransform.SetRotation(UpdatedRotation);
 	}
 
+	UpdateElementTransforms(PreviewElements, FinalizedTransform);
+}
+
+void UPlacementModePlaceSingleTool::DestroyPreviewElements()
+{
+	NotifyMovementEnded(PreviewElements);
+	PreviewElements.Empty();
+	ShutdownBrushStampIndicator();
+}
+
+void UPlacementModePlaceSingleTool::EnterTweakState(TArrayView<const FTypedElementHandle> InElementHandles)
+{
+	if (InElementHandles.Num() == 0)
+	{
+		return;
+	}
+
+	FToolBuilderState SelectionState;
+	GetToolManager()->GetContextQueriesAPI()->GetCurrentSelectionState(SelectionState);
+	SelectionState.TypedElementSelectionSet->SetSelection(InElementHandles, FTypedElementSelectionOptions());
+	bIsTweaking = true;
+}
+
+void UPlacementModePlaceSingleTool::ExitTweakState(bool bClearSelectionSet)
+{
+	if (bClearSelectionSet)
+	{
+		FToolBuilderState SelectionState;
+		GetToolManager()->GetContextQueriesAPI()->GetCurrentSelectionState(SelectionState);
+		SelectionState.TypedElementSelectionSet->ClearSelection(FTypedElementSelectionOptions());
+	}
+
+	bIsTweaking = false;
+}
+
+void UPlacementModePlaceSingleTool::UpdateElementTransforms(TArrayView<const FTypedElementHandle> InElements, const FTransform& InTransform, bool bLocalTransform)
+{
 	// Update the transform positions for the preview elements.
-	for (FTypedElementHandle& PreviewElement : PreviewElements)
+	for (const FTypedElementHandle& ElementHandle : InElements)
 	{
-		if (TTypedElement<UTypedElementWorldInterface> WorldInterfaceElement = UTypedElementRegistry::GetInstance()->GetElement<UTypedElementWorldInterface>(PreviewElement))
+		if (TTypedElement<UTypedElementWorldInterface> WorldInterfaceElement = UTypedElementRegistry::GetInstance()->GetElement<UTypedElementWorldInterface>(ElementHandle))
 		{
-			WorldInterfaceElement.SetWorldTransform(FinalizedTransform);
+			bLocalTransform ? WorldInterfaceElement.SetRelativeTransform(InTransform) : WorldInterfaceElement.SetWorldTransform(InTransform);
 			WorldInterfaceElement.NotifyMovementOngoing();
 		}
 	}
 }
 
-void UPlacementModePlaceSingleTool::DestroyPreviewElements()
+void UPlacementModePlaceSingleTool::NotifyMovementStarted(TArrayView<const FTypedElementHandle> InElements)
+{
+	// Notify Movement started
+	for (const FTypedElementHandle& PreviewElement : PreviewElements)
+	{
+		if (TTypedElement<UTypedElementWorldInterface> WorldInterfaceElement = UTypedElementRegistry::GetInstance()->GetElement<UTypedElementWorldInterface>(PreviewElement))
+		{
+			WorldInterfaceElement.NotifyMovementStarted();
+		}
+	}
+}
+
+void UPlacementModePlaceSingleTool::NotifyMovementEnded(TArrayView<const FTypedElementHandle> InElements)
 {
 	FToolBuilderState SelectionState;
 	GetToolManager()->GetContextQueriesAPI()->GetCurrentSelectionState(SelectionState);
 	check(SelectionState.TypedElementSelectionSet.IsValid());	// Placement tools expect a valid selection set.
 
-	for (FTypedElementHandle& PreviewElement : PreviewElements)
+	for (const FTypedElementHandle& PreviewElement : PreviewElements)
 	{
 		if (TTypedElement<UTypedElementWorldInterface> WorldInterfaceElement = UTypedElementRegistry::GetInstance()->GetElement<UTypedElementWorldInterface>(PreviewElement))
 		{
@@ -216,6 +261,20 @@ void UPlacementModePlaceSingleTool::DestroyPreviewElements()
 			WorldInterfaceElement.DeleteElement(WorldInterfaceElement.GetOwnerWorld(), SelectionState.TypedElementSelectionSet.Get(), FTypedElementDeletionOptions());
 		}
 	}
+}
 
-	PreviewElements.Empty();
+void UPlacementModePlaceSingleTool::SetupRightClickMouseBehavior()
+{
+	ULocalClickDragInputBehavior* RightMouseBehavior = NewObject<ULocalClickDragInputBehavior>(this);
+	RightMouseBehavior->CanBeginClickDragFunc = [this](const FInputDeviceRay&) { return bShiftToggle ? FInputRayHit(1.0f) : FInputRayHit(); };
+	RightMouseBehavior->OnClickReleaseFunc = [this](const FInputDeviceRay&)
+	{
+		constexpr bool bClearSelection = true;
+		ExitTweakState(bClearSelection);
+		PreviewPlacementInfo.Reset();
+	};
+	RightMouseBehavior->SetDefaultPriority(FInputCapturePriority(-1));
+	RightMouseBehavior->SetUseRightMouseButton();
+	RightMouseBehavior->Initialize();
+	AddInputBehavior(RightMouseBehavior);
 }
diff --git a/Engine/Source/Editor/AssetPlacementEdMode/Private/Tools/PlacementPlaceSingleTool.h b/Engine/Source/Editor/AssetPlacementEdMode/Private/Tools/PlacementPlaceSingleTool.h
--- a/Engine/Source/Editor/AssetPlacementEdMode/Private/Tools/PlacementPlaceSingleTool.h
+++ b/Engine/Source/Editor/AssetPlacementEdMode/Private/Tools/PlacementPlaceSingleTool.h
@@ -2,7 +2,6 @@
 
 #pragma once
 
-#include "InteractiveToolBuilder.h"
 #include "Tools/PlacementClickDragToolBase.h"
 
 #include "PlacementPlaceSingleTool.generated.h"
@@ -18,36 +17,6 @@ protected:
 	virtual UPlacementBrushToolBase* FactoryToolInstance(UObject* Outer) const override;
 };
 
-UCLASS(config = EditorPerUserSettings)
-class UPlacementModePlaceSingleSettings : public UInteractiveToolPropertySet
-{
-	GENERATED_BODY()
-
-public:
-	UPROPERTY(config, EditAnywhere, Category = "Single Place Tool", meta = (InlineEditConditionToggle))
-	bool bAlignToCursor = false;
-
-	UPROPERTY(config, EditAnywhere, Category = "Single Place Tool", meta = (DisplayName = "Align to Cursor", EditCondition = "bAlignToCursor"))
-	TEnumAsByte<EAxis::Type> AxisToAlignWithCursor = EAxis::Type::X;
-
-	UPROPERTY(config, EditAnywhere, Category = "Single Place Tool")
-	bool bInvertCursorAxis = false;
-
-	UPROPERTY(config, EditAnywhere, Category = "Single Place Tool")
-	bool bSnapToGridX = false;
-
-	UPROPERTY(config, EditAnywhere, Category = "Single Place Tool")
-	bool bSnapToGridY = false;
-
-	UPROPERTY(config, EditAnywhere, Category = "Single Place Tool")
-	bool bSnapToGridZ = false;
-
-	UPROPERTY(config, EditAnywhere, Category = "Single Place Tool")
-	float ScrollWheelOffsetScale = 0.05f;
-
-	virtual bool CanEditChange(const FProperty* Property) const override;
-};
-
 UCLASS(Transient)
 class UPlacementModePlaceSingleTool : public UPlacementClickDragToolBase
 {
@@ -67,16 +36,22 @@ public:
 	virtual void OnBeginHover(const FInputDeviceRay& DevicePos) override;
 	virtual bool OnUpdateHover(const FInputDeviceRay& DevicePos) override;
 	virtual void OnEndHover() override;
+	virtual FInputRayHit CanBeginClickDragSequence(const FInputDeviceRay& PressPos) override;
+	virtual FInputRayHit BeginHoverSequenceHitTest(const FInputDeviceRay& DevicePos) override;
 
 protected:
 	void GeneratePreviewPlacementData(const FInputDeviceRay& DevicePos);
 	void CreatePreviewElements(const FInputDeviceRay& DevicePos);
 	void UpdatePreviewElements(const FInputDeviceRay& DevicePos);
 	void DestroyPreviewElements();
-
-	UPROPERTY()
-	TObjectPtr<UPlacementModePlaceSingleSettings> SingleToolSettings;
+	void EnterTweakState(TArrayView<const FTypedElementHandle> InElementHandles);
+	void ExitTweakState(bool bClearSelectionSet);
+	void UpdateElementTransforms(TArrayView<const FTypedElementHandle> InElements, const FTransform& InTransform, bool bLocalTransform = false);
+	void NotifyMovementStarted(TArrayView<const FTypedElementHandle> InElements);
+	void NotifyMovementEnded(TArrayView<const FTypedElementHandle> InElements);
+	void SetupRightClickMouseBehavior();
 
 	TUniquePtr<FAssetPlacementInfo> PreviewPlacementInfo;
 	TArray<FTypedElementHandle> PreviewElements;
+	bool bIsTweaking;
 };
diff --git a/Engine/Source/Editor/AssetPlacementEdMode/Private/Tools/PlacementSelectTool.cpp b/Engine/Source/Editor/AssetPlacementEdMode/Private/Tools/PlacementSelectTool.cpp
--- a/Engine/Source/Editor/AssetPlacementEdMode/Private/Tools/PlacementSelectTool.cpp
+++ b/Engine/Source/Editor/AssetPlacementEdMode/Private/Tools/PlacementSelectTool.cpp
@@ -12,7 +12,18 @@ bool UPlacementModeSelectToolBuilder::CanBuildTool(const FToolBuilderState& Scen
 	return true;
 }
 
-UInteractiveTool* UPlacementModeSelectToolBuilder::BuildTool(const FToolBuilderState& SceneState) const
+UPlacementBrushToolBase* UPlacementModeSelectToolBuilder::FactoryToolInstance(UObject* Outer) const
 {
-	return NewObject<UPlacementModeSelectTool>(SceneState.ToolManager, UPlacementModeSelectTool::ToolName);
+	return NewObject<UPlacementModeSelectTool>(Outer);
+}
+
+void UPlacementModeSelectTool::Setup()
+{
+	Super::Setup();
+	ShutdownBrushStampIndicator();
+}
+
+FInputRayHit UPlacementModeSelectTool::CanBeginClickDragSequence(const FInputDeviceRay& PressPos)
+{
+	return FInputRayHit();	// always fall-through to the EdMode or viewport handling for clicks.
 }
diff --git a/Engine/Source/Editor/AssetPlacementEdMode/Private/Tools/PlacementSelectTool.h b/Engine/Source/Editor/AssetPlacementEdMode/Private/Tools/PlacementSelectTool.h
--- a/Engine/Source/Editor/AssetPlacementEdMode/Private/Tools/PlacementSelectTool.h
+++ b/Engine/Source/Editor/AssetPlacementEdMode/Private/Tools/PlacementSelectTool.h
@@ -2,19 +2,21 @@
 
 #pragma once
 
-#include "InteractiveToolBuilder.h"
+#include "Tools/PlacementBrushToolBase.h"
 #include "Tools/PlacementClickDragToolBase.h"
 
 #include "PlacementSelectTool.generated.h"
 
 UCLASS(Transient, MinimalAPI)
-class UPlacementModeSelectToolBuilder : public UInteractiveToolBuilder
+class UPlacementModeSelectToolBuilder : public UPlacementToolBuilderBase
 {
 	GENERATED_BODY()
 
 public:
 	virtual bool CanBuildTool(const FToolBuilderState& SceneState) const override;
-	virtual UInteractiveTool* BuildTool(const FToolBuilderState& SceneState) const override;
+
+protected:
+	virtual UPlacementBrushToolBase* FactoryToolInstance(UObject* Outer) const override;
 };
 
 UCLASS(MinimalAPI)
@@ -23,5 +25,8 @@ class UPlacementModeSelectTool : public UPlacementClickDragToolBase
 	GENERATED_BODY()
 
 public:
+	virtual void Setup() override;
+	virtual FInputRayHit CanBeginClickDragSequence(const FInputDeviceRay& PressPos) override;
+
 	constexpr static TCHAR ToolName[] = TEXT("SelectTool");
 };
diff --git a/Engine/Source/Editor/AssetPlacementEdMode/Public/AssetPlacementEdMode.h b/Engine/Source/Editor/AssetPlacementEdMode/Public/AssetPlacementEdMode.h
--- a/Engine/Source/Editor/AssetPlacementEdMode/Public/AssetPlacementEdMode.h
+++ b/Engine/Source/Editor/AssetPlacementEdMode/Public/AssetPlacementEdMode.h
@@ -3,7 +3,7 @@
 #pragma once
 
 #include "CoreMinimal.h"
-#include "Tools/LegacyEdModeWidgetHelpers.h"
+#include "Tools/DefaultEdMode.h"
 #include "Elements/Framework/TypedElementHandle.h"
 
 #include "AssetPlacementEdMode.generated.h"
@@ -11,7 +11,7 @@
 class UAssetPlacementSettings;
 
 UCLASS()
-class UAssetPlacementEdMode : public UBaseLegacyWidgetEdMode
+class UAssetPlacementEdMode : public UEdModeDefault
 {
 	GENERATED_BODY()
 
@@ -26,10 +26,12 @@ public:
 	virtual void Enter() override;
 	virtual void Exit() override;
 	virtual void CreateToolkit() override;
+	virtual bool UsesToolkits() const override;
 	virtual TMap<FName, TArray<TSharedPtr<FUICommandInfo>>> GetModeCommands() const override;
 	virtual void BindCommands() override;
 	virtual bool IsSelectionAllowed(AActor* InActor, bool bInSelection) const override;
 	virtual void OnToolStarted(UInteractiveToolManager* Manager, UInteractiveTool* Tool) override;
+	virtual void OnToolEnded(UInteractiveToolManager* Manager, UInteractiveTool* Tool) override;
 	// End of UEdMode interface
 	//////////////////
 
diff --git a/Engine/Source/Editor/AssetPlacementEdMode/Public/AssetPlacementEdModeToolkit.h b/Engine/Source/Editor/AssetPlacementEdMode/Public/AssetPlacementEdModeToolkit.h
--- a/Engine/Source/Editor/AssetPlacementEdMode/Public/AssetPlacementEdModeToolkit.h
+++ b/Engine/Source/Editor/AssetPlacementEdMode/Public/AssetPlacementEdModeToolkit.h
@@ -4,9 +4,6 @@
 
 #include "CoreMinimal.h"
 #include "Toolkits/BaseToolkit.h"
-#include "AssetPlacementEdMode.h"
-
-class UAssetPlacementSettings;
 
 /**
  * This FModeToolkit just creates a basic UI panel that allows various InteractiveTools to
@@ -15,7 +12,7 @@ class UAssetPlacementSettings;
 class FAssetPlacementEdModeToolkit : public FModeToolkit
 {
 public:
-	FAssetPlacementEdModeToolkit(TWeakObjectPtr<UAssetPlacementSettings> InPlacementSettings);
+	FAssetPlacementEdModeToolkit();
 	
 	// FModeToolkit interface 
 	virtual void GetToolPaletteNames(TArray<FName>& PaletteNames) const override;
@@ -26,6 +23,4 @@ public:
 
 protected:
 	virtual TSharedPtr<SWidget> GetInlineContent() const override;
-
-	TWeakObjectPtr<UAssetPlacementSettings> PlacementSettings;
 };
diff --git a/Engine/Source/Editor/EditorFramework/Private/Subsystems/PlacementSubsystem.cpp b/Engine/Source/Editor/EditorFramework/Private/Subsystems/PlacementSubsystem.cpp
--- a/Engine/Source/Editor/EditorFramework/Private/Subsystems/PlacementSubsystem.cpp
+++ b/Engine/Source/Editor/EditorFramework/Private/Subsystems/PlacementSubsystem.cpp
@@ -30,8 +30,7 @@ TArray<FTypedElementHandle> UPlacementSubsystem::PlaceAsset(const FAssetPlacemen
 
 TArray<FTypedElementHandle> UPlacementSubsystem::PlaceAssets(TArrayView<const FAssetPlacementInfo> InPlacementInfos, const FPlacementOptions& InPlacementOptions)
 {
-	bool bWasCreatingPreviewElements = bIsCreatingPreviewElements;
-	bIsCreatingPreviewElements = InPlacementOptions.bIsCreatingPreviewElements;
+	TGuardValue<bool> bShouldCreatePreviewElements(bIsCreatingPreviewElements, InPlacementOptions.bIsCreatingPreviewElements);
 
 	TArray<FTypedElementHandle> PlacedElements;
 	for (const FAssetPlacementInfo& PlacementInfo : InPlacementInfos)
@@ -62,8 +61,6 @@ TArray<FTypedElementHandle> UPlacementSubsystem::PlaceAssets(TArrayView<const FA
 		}
 	}
 
-	bIsCreatingPreviewElements = bWasCreatingPreviewElements;
-
 	return PlacedElements;
 }
 
diff --git a/Engine/Source/Editor/EditorFramework/Public/Subsystems/PlacementSubsystem.h b/Engine/Source/Editor/EditorFramework/Public/Subsystems/PlacementSubsystem.h
--- a/Engine/Source/Editor/EditorFramework/Public/Subsystems/PlacementSubsystem.h
+++ b/Engine/Source/Editor/EditorFramework/Public/Subsystems/PlacementSubsystem.h
@@ -84,7 +84,10 @@ public:
 	 */
 	TScriptInterface<IAssetFactoryInterface> FindAssetFactoryFromAssetData(const FAssetData& InAssetData);
 
-	// @returns true if the current PlaceAssets call is creating preview elements.
+	/**
+	 * Determines if the placement system is currently placing preview elements.
+	 * @returns true if the current PlaceAssets call is creating preview elements.
+	 */
 	bool IsCreatingPreviewElements() const;
 
 private:
diff --git a/Engine/Source/Editor/Experimental/EditorInteractiveToolsFramework/Private/EdModeInteractiveToolsContext.cpp b/Engine/Source/Editor/Experimental/EditorInteractiveToolsFramework/Private/EdModeInteractiveToolsContext.cpp
--- a/Engine/Source/Editor/Experimental/EditorInteractiveToolsFramework/Private/EdModeInteractiveToolsContext.cpp
+++ b/Engine/Source/Editor/Experimental/EditorInteractiveToolsFramework/Private/EdModeInteractiveToolsContext.cpp
@@ -523,6 +523,7 @@ UEdModeInteractiveToolsContext::UEdModeInteractiveToolsContext()
 	QueriesAPI = nullptr;
 	TransactionAPI = nullptr;
 	AssetAPI = nullptr;
+	bIsTrackingMouse = false;
 }
 
 
@@ -1099,15 +1100,8 @@ bool UEdModeInteractiveToolsContext::MouseLeave(FEditorViewportClient* ViewportC
 
 bool UEdModeInteractiveToolsContext::StartTracking(FEditorViewportClient* InViewportClient, FViewport* InViewport)
 {
-	// capture tracking if we have an active tool
-	if (ToolManager->HasActiveTool(EToolSide::Mouse))
-	{
-#ifdef ENABLE_DEBUG_PRINTING
-		UE_LOG(LogTemp, Warning, TEXT("BEGIN TRACKING"));
-#endif
-		return true;
-	}
-	return false;
+	bIsTrackingMouse = InputRouter->HasActiveMouseCapture();
+	return bIsTrackingMouse;
 }
 
 bool UEdModeInteractiveToolsContext::CapturedMouseMove(FEditorViewportClient* InViewportClient, FViewport* InViewport, int32 InMouseX, int32 InMouseY)
@@ -1146,11 +1140,9 @@ bool UEdModeInteractiveToolsContext::EndTracking(FEditorViewportClient* InViewpo
 	// unlock flight camera
 	InViewportClient->bLockFlightCamera = false;
 
-	if (ToolManager->HasActiveTool(EToolSide::Mouse))
+	if (bIsTrackingMouse)
 	{
-#ifdef ENABLE_DEBUG_PRINTING
-	UE_LOG(LogTemp, Warning, TEXT("END TRACKING"));
-#endif
+		bIsTrackingMouse = false;
 		return true;
 	}
 
diff --git a/Engine/Source/Editor/Experimental/EditorInteractiveToolsFramework/Public/EdModeInteractiveToolsContext.h b/Engine/Source/Editor/Experimental/EditorInteractiveToolsFramework/Public/EdModeInteractiveToolsContext.h
--- a/Engine/Source/Editor/Experimental/EditorInteractiveToolsFramework/Public/EdModeInteractiveToolsContext.h
+++ b/Engine/Source/Editor/Experimental/EditorInteractiveToolsFramework/Public/EdModeInteractiveToolsContext.h
@@ -151,5 +151,5 @@ protected:
 
 private:
 	FEditorModeTools* EditorModeManager = nullptr;
-
+	bool bIsTrackingMouse;
 };
diff --git a/Engine/Source/Editor/LevelEditor/Private/SLevelEditorToolBox.cpp b/Engine/Source/Editor/LevelEditor/Private/SLevelEditorToolBox.cpp
--- a/Engine/Source/Editor/LevelEditor/Private/SLevelEditorToolBox.cpp
+++ b/Engine/Source/Editor/LevelEditor/Private/SLevelEditorToolBox.cpp
@@ -272,13 +272,18 @@ void SLevelEditorToolBox::UpdateInlineContent(const TSharedPtr<IToolkit>& Toolki
 				GEditor->GetEditorSubsystem<UStatusBarSubsystem>()->PopStatusBarMessage(LevelEditorStatusBarName, StatusBarMessageHandle);
 			}
 
-			StatusBarMessageHandle = GEditor->GetEditorSubsystem<UStatusBarSubsystem>()->PushStatusBarMessage(LevelEditorStatusBarName, TAttribute<FText>::Create(
-				TAttribute<FText>::FGetter::CreateLambda([ModeToolkit]()
+			StatusBarMessageHandle = GEditor->GetEditorSubsystem<UStatusBarSubsystem>()->PushStatusBarMessage(LevelEditorStatusBarName, MakeAttributeLambda([ModeToolkit]()
+				{
+					TSharedPtr<FModeToolkit> PinnedToolkit = ModeToolkit.Pin();
+					if (PinnedToolkit.IsValid())
 					{
-						TSharedPtr<FModeToolkit> PinnedToolkit = ModeToolkit.Pin();
-						if (PinnedToolkit.IsValid()) { return PinnedToolkit->GetActiveToolDisplayName(); }
-						else { return FText(); }
-					}))
+						return PinnedToolkit->GetActiveToolDisplayName();
+					}
+					else
+					{
+						return FText();
+					}
+				})
 			);
 		}
 	}
diff --git a/Engine/Source/Runtime/Experimental/InteractiveToolsFramework/Private/BaseTools/BaseBrushTool.cpp b/Engine/Source/Runtime/Experimental/InteractiveToolsFramework/Private/BaseTools/BaseBrushTool.cpp
--- a/Engine/Source/Runtime/Experimental/InteractiveToolsFramework/Private/BaseTools/BaseBrushTool.cpp
+++ b/Engine/Source/Runtime/Experimental/InteractiveToolsFramework/Private/BaseTools/BaseBrushTool.cpp
@@ -232,21 +232,30 @@ const FString BaseBrushIndicatorGizmoType = TEXT("BrushIndicatorGizmoType");
 
 void UBaseBrushTool::SetupBrushStampIndicator()
 {
-	// register and spawn brush indicator gizmo
-	GetToolManager()->GetPairedGizmoManager()->RegisterGizmoType(BaseBrushIndicatorGizmoType, NewObject<UBrushStampIndicatorBuilder>());
-	BrushStampIndicator = GetToolManager()->GetPairedGizmoManager()->CreateGizmo<UBrushStampIndicator>(BaseBrushIndicatorGizmoType, FString(), this);
+	if (!BrushStampIndicator)
+	{
+		// register and spawn brush indicator gizmo
+		GetToolManager()->GetPairedGizmoManager()->RegisterGizmoType(BaseBrushIndicatorGizmoType, NewObject<UBrushStampIndicatorBuilder>());
+		BrushStampIndicator = GetToolManager()->GetPairedGizmoManager()->CreateGizmo<UBrushStampIndicator>(BaseBrushIndicatorGizmoType, FString(), this);
+	}
 }
 
 void UBaseBrushTool::UpdateBrushStampIndicator()
 {
-	BrushStampIndicator->Update(LastBrushStamp.Radius, LastBrushStamp.WorldPosition, LastBrushStamp.WorldNormal, LastBrushStamp.Falloff);
+	if (BrushStampIndicator)
+	{
+		BrushStampIndicator->Update(LastBrushStamp.Radius, LastBrushStamp.WorldPosition, LastBrushStamp.WorldNormal, LastBrushStamp.Falloff);
+	}
 }
 
 void UBaseBrushTool::ShutdownBrushStampIndicator()
 {
-	GetToolManager()->GetPairedGizmoManager()->DestroyAllGizmosByOwner(this);
-	BrushStampIndicator = nullptr;
-	GetToolManager()->GetPairedGizmoManager()->DeregisterGizmoType(BaseBrushIndicatorGizmoType);
+	if (BrushStampIndicator)
+	{
+		GetToolManager()->GetPairedGizmoManager()->DestroyGizmo(BrushStampIndicator);
+		BrushStampIndicator = nullptr;
+		GetToolManager()->GetPairedGizmoManager()->DeregisterGizmoType(BaseBrushIndicatorGizmoType);
+	}
 }
 
 
