Turn off undo/redo tracking for switching between tools in placement mode.

This was causing a crash when undoing due to the selection set being cleared on tool enter, and it felt bad to be in a different tool than I last selected after undoing.

diff --git a/Engine/Source/Editor/AssetPlacementEdMode/Private/AssetPlacementEdMode.cpp b/Engine/Source/Editor/AssetPlacementEdMode/Private/AssetPlacementEdMode.cpp
--- a/Engine/Source/Editor/AssetPlacementEdMode/Private/AssetPlacementEdMode.cpp
+++ b/Engine/Source/Editor/AssetPlacementEdMode/Private/AssetPlacementEdMode.cpp
@@ -66,11 +66,11 @@ void UAssetPlacementEdMode::Enter()
 	Owner->StoreSelection(AssetPlacementEdModeID);
 	bIsInSelectionTool = false;
 
-	// Enable the select tool by default.
 	// Disable undo tracking so that we can't accidentally undo ourselves out of the select mode and into an invalid state.
 	GetToolManager()->ConfigureChangeTrackingMode(EToolChangeTrackingMode::NoChangeTracking);
+
+	// Enable the select tool by default.
 	ToolsContext->StartTool(UPlacementModeSelectTool::ToolName);
-	GetToolManager()->ConfigureChangeTrackingMode(EToolChangeTrackingMode::FullUndoRedo);
 }
 
 void UAssetPlacementEdMode::Exit()
