Non-unity compile fix.

#rnx

diff --git a/Engine/Source/Editor/AssetPlacementEdMode/Private/Tools/PlacementLassoSelectTool.cpp b/Engine/Source/Editor/AssetPlacementEdMode/Private/Tools/PlacementLassoSelectTool.cpp
--- a/Engine/Source/Editor/AssetPlacementEdMode/Private/Tools/PlacementLassoSelectTool.cpp
+++ b/Engine/Source/Editor/AssetPlacementEdMode/Private/Tools/PlacementLassoSelectTool.cpp
@@ -4,6 +4,7 @@
 
 #include "AssetPlacementEdMode.h"
 #include "AssetPlacementSettings.h"
+#include "Editor.h"
 #include "Elements/Framework/EngineElementsLibrary.h"
 #include "Elements/Framework/TypedElementHandle.h"
 #include "Elements/Framework/TypedElementRegistry.h"
