Hooked up remaining placement tool settings (world offset, local offset, negative scaling, negative rotation).

Disabled ISM placement in single place tool, since single selection of ISM is not hooked up yet.

diff --git a/Engine/Source/Editor/AssetPlacementEdMode/Private/AssetPlacementSettings.cpp b/Engine/Source/Editor/AssetPlacementEdMode/Private/AssetPlacementSettings.cpp
--- a/Engine/Source/Editor/AssetPlacementEdMode/Private/AssetPlacementSettings.cpp
+++ b/Engine/Source/Editor/AssetPlacementEdMode/Private/AssetPlacementSettings.cpp
@@ -10,28 +10,24 @@ bool UAssetPlacementSettings::CanEditChange(const FProperty* InProperty) const
 	}
 
 	const FName PropertyName = InProperty->GetFName();
-	if (PropertyName == GET_MEMBER_NAME_CHECKED(UAssetPlacementSettings, ScaleRangeUniform) || PropertyName == GET_MEMBER_NAME_CHECKED(UAssetPlacementSettings, bAllowNegativeUniformScale))
+	if (PropertyName == GET_MEMBER_NAME_CHECKED(UAssetPlacementSettings, ScaleRange) || PropertyName == GET_MEMBER_NAME_CHECKED(UAssetPlacementSettings, bAllowNegativeScale))
 	{
-		return (ScalingType == EFoliageScaling::Uniform) && bUseRandomScale;
+		return bUseRandomScale;
 	}
 
-	const bool bCanFreeScale = (ScalingType != EFoliageScaling::Uniform) && bUseRandomScale;
-	if (PropertyName == GET_MEMBER_NAME_CHECKED(UAssetPlacementSettings, ScaleRangeX) || PropertyName == GET_MEMBER_NAME_CHECKED(UAssetPlacementSettings, bAllowNegativeXScale))
+	if (PropertyName == GET_MEMBER_NAME_CHECKED(UAssetPlacementSettings, bAllowNegativeRotationX))
 	{
-		const bool bLockX = (ScalingType == EFoliageScaling::LockXY) || (ScalingType == EFoliageScaling::LockXZ);
-		return bCanFreeScale && !bLockX;
+		return bUseRandomRotationX;
 	}
 
-	if (PropertyName == GET_MEMBER_NAME_CHECKED(UAssetPlacementSettings, ScaleRangeY) || PropertyName == GET_MEMBER_NAME_CHECKED(UAssetPlacementSettings, bAllowNegativeYScale))
+	if (PropertyName == GET_MEMBER_NAME_CHECKED(UAssetPlacementSettings, bAllowNegativeRotationY))
 	{
-		const bool bLockY = (ScalingType == EFoliageScaling::LockXY) || (ScalingType == EFoliageScaling::LockYZ);
-		return bCanFreeScale && !bLockY;
+		return bUseRandomRotationY;
 	}
 
-	if (PropertyName == GET_MEMBER_NAME_CHECKED(UAssetPlacementSettings, FreeScaleRangeZ) || PropertyName == GET_MEMBER_NAME_CHECKED(UAssetPlacementSettings, bAllowNegativeZScale))
+	if (PropertyName == GET_MEMBER_NAME_CHECKED(UAssetPlacementSettings, bAllowNegativeRotationZ))
 	{
-		const bool bLockZ = (ScalingType == EFoliageScaling::LockYZ) || (ScalingType == EFoliageScaling::LockXZ);
-		return bCanFreeScale && !bLockZ;
+		return bUseRandomRotationZ;
 	}
 
 	if (PropertyName == GET_MEMBER_NAME_CHECKED(UAssetPlacementSettings, bInvertNormalAxis))
diff --git a/Engine/Source/Editor/AssetPlacementEdMode/Private/AssetPlacementSettings.h b/Engine/Source/Editor/AssetPlacementEdMode/Private/AssetPlacementSettings.h
--- a/Engine/Source/Editor/AssetPlacementEdMode/Private/AssetPlacementSettings.h
+++ b/Engine/Source/Editor/AssetPlacementEdMode/Private/AssetPlacementSettings.h
@@ -28,80 +28,107 @@ class UAssetPlacementSettings : public UObject
 	GENERATED_BODY()
 
 public:
+	/** Set to true if the tool should automatically align placement to an axis. */
 	UPROPERTY(config, EditAnywhere, Category = "Axis Alignment", meta = (InlineEditConditionToggle))
 	bool bAlignToNormal = true;
-	
+
+	/** The axis which the tool should use for alignment calculations. This axis is positive unless Invert Normal Axis is set. */
 	UPROPERTY(config, EditAnywhere, Category = "Axis Alignment", meta = (DisplayName = "Align To Normal", EditCondition = "bAlignToNormal"))
 	TEnumAsByte<EAxis::Type> AxisToAlignWithNormal = EAxis::Type::Z;
 
+	/** Set to true if the tool should negate the alignment axis. */
 	UPROPERTY(config, EditAnywhere, Category = "Axis Alignment")
 	bool bInvertNormalAxis = false;
-	
+
+	/** Set to true if the tool should generate a random rotation on the X axis */
 	UPROPERTY(config, EditAnywhere, Category = "Rotation", meta = (InlineEditConditionToggle))
-	bool bUseRandomXRotation = false;
+	bool bUseRandomRotationX = false;
 
-	UPROPERTY(config, EditAnywhere, AdvancedDisplay, Category = "Rotation", meta = (EditCondition = "bUseRandomXRotation", UIMin = "0.0", UIMax = "360.0"))
-	FFloatInterval RandomRotationX = FFloatInterval(0.0f, 360.0f);
-	
+	/** The range in degrees at which the tool should generate a random rotation on the X axis. */
+	UPROPERTY(config, EditAnywhere, AdvancedDisplay, Category = "Rotation", meta = (EditCondition = "bUseRandomRotationX", UIMin = "0.0", UIMax = "180.0"))
+	FFloatInterval RandomRotationX = FFloatInterval(0.0f, 180.0f);
+
+	/** Set to true if the tool should generate a negatively signed rotation on the X axis. */
+	UPROPERTY(config, EditAnywhere, AdvancedDisplay, Category = "Rotation")
+	bool bAllowNegativeRotationX = true;
+
+	/** Set to true if the tool should generate a random rotation on the Y axis */
 	UPROPERTY(config, EditAnywhere, Category = "Rotation", meta = (InlineEditConditionToggle))
-	bool bUseRandomYRotation = false;
+	bool bUseRandomRotationY = false;
+
+	/** The range in degrees at which the tool should generate a random rotation on the Y axis. */
+	UPROPERTY(config, EditAnywhere, AdvancedDisplay, Category = "Rotation", meta = (EditCondition = "bUseRandomRotationY", UIMin = "0.0", UIMax = "180.0"))
+	FFloatInterval RandomRotationY = FFloatInterval(0.0f, 180.0f);
 
-	UPROPERTY(config, EditAnywhere, AdvancedDisplay, Category = "Rotation", meta = (EditCondition = "bUseRandomYRotation", UIMin = "0.0", UIMax = "360.0"))
-	FFloatInterval RandomRotationY = FFloatInterval(0.0f, 360.0f);
+	/** Set to true if the tool should generate a negatively signed rotation on the Y axis. */
+	UPROPERTY(config, EditAnywhere, AdvancedDisplay, Category = "Rotation")
+	bool bAllowNegativeRotationY = true;
 
+	/** Set to true if the tool should generate a random rotation on the Z axis */
 	UPROPERTY(config, EditAnywhere, Category = "Rotation", meta = (InlineEditConditionToggle))
-	bool bUseRandomZRotation = true;
+	bool bUseRandomRotationZ = true;
 
-	UPROPERTY(config, EditAnywhere, Category = "Rotation", meta = (EditCondition = "bUseRandomZRotation", UIMin = "0.0", UIMax = "360.0"))
-	FFloatInterval RandomRotationZ = FFloatInterval(0.0f, 360.0f);
+	/** The range in degrees at which the tool should generate a random rotation on the Z axis. */
+	UPROPERTY(config, EditAnywhere, Category = "Rotation", meta = (EditCondition = "bUseRandomRotationZ", UIMin = "0.0", UIMax = "180.0"))
+	FFloatInterval RandomRotationZ = FFloatInterval(0.0f, 180.0f);
+	
+	/** Set to true if the tool should generate a negatively signed rotation on the Z axis. */
+	UPROPERTY(config, EditAnywhere, Category = "Rotation")
+	bool bAllowNegativeRotationZ = true;
 
+	/** Set to true if the tool should randomly scale any generated assets */
 	UPROPERTY(config, EditAnywhere, Category = "Scale", meta = (InlineEditConditionToggle))
 	bool bUseRandomScale = true;
 
-	UPROPERTY(config, EditAnywhere, Category = "Scale", meta = (DisplayName="Use Random Scale", EditCondition = "bAllowRandomScale"))
+	/** Which axes any randomly generated scale should be applied to. Any locked axes will be scaled to 1.0f. */
+	UPROPERTY(config, EditAnywhere, Category = "Scale", meta = (DisplayName="Use Random Scale", EditCondition = "bUseRandomScale"))
 	EFoliageScaling ScalingType = EFoliageScaling::Uniform;
 	
-	UPROPERTY(config, EditAnywhere, AdvancedDisplay, Category = "Scale", meta = (UIMin = ".01"))
-	FFloatInterval ScaleRangeUniform = FFloatInterval(.8f, 1.0f);
-	
-	UPROPERTY(config, EditAnywhere, AdvancedDisplay, Category = "Scale")
-	bool bAllowNegativeUniformScale = false;
-	
-	UPROPERTY(config, EditAnywhere, AdvancedDisplay, Category = "Scale", meta = (UIMin = ".01"))
-	FFloatInterval ScaleRangeX = FFloatInterval(.8f, 1.0f);
-	
-	UPROPERTY(config, EditAnywhere, AdvancedDisplay, Category = "Scale")
-	bool bAllowNegativeXScale = false;
+	/** The min and max at which a random scale is generated. */
+	UPROPERTY(config, EditAnywhere, Category = "Scale", meta = (UIMin = ".01"))
+	FFloatInterval ScaleRange = FFloatInterval(.5f, 2.0f);
 
-	UPROPERTY(config, EditAnywhere, AdvancedDisplay, Category = "Scale", meta = (UIMin = ".01"))
-	FFloatInterval ScaleRangeY = FFloatInterval(.8f, 1.0f);
-	
+	/** Set to true if the tool should generate a negatively signed scale. */
 	UPROPERTY(config, EditAnywhere, AdvancedDisplay, Category = "Scale")
-	bool bAllowNegativeYScale = false;
+	bool bAllowNegativeScale = false;
 
-	UPROPERTY(config, EditAnywhere, AdvancedDisplay, Category = "Scale", meta = (UIMin = ".01"))
-	FFloatInterval FreeScaleRangeZ = FFloatInterval(.8f, 1.0f);
-	
-	UPROPERTY(config, EditAnywhere, AdvancedDisplay, Category = "Scale")
-	bool bAllowNegativeZScale = false;
+	/** The spawn location offset in relative space. */
+	UPROPERTY(config, EditAnywhere, Category = "Location")
+	FVector RelativeLocationOffset = FVector::ZeroVector;
+
+	/** If the relative location offset should match scale with the placed asset. */
+	UPROPERTY(config, EditAnywhere, Category = "Location")
+	bool bScaleRelativeLocationOffset = true;
+
+	/** The spawn rotation offset in world space. */
+	UPROPERTY(config, EditAnywhere, AdvancedDisplay, Category = "Location")
+	FVector WorldLocationOffset = FVector::ZeroVector;
+
+	/** If the world location offset should match scale with the placed asset. */
+	UPROPERTY(config, EditAnywhere, AdvancedDisplay, Category = "Location")
+	bool bScaleWorldLocationOffset = false;
 
+	/** If the tool should place on landscape surfaces. */
 	UPROPERTY(config, EditAnywhere, Category = "Place on Filters")
 	bool bLandscape = true;
 
+	/** If the tool should place on static mesh surfaces. */
 	UPROPERTY(config, EditAnywhere, Category = "Place on Filters")
 	bool bStaticMeshes = true;
 
+	/** If the tool should place on BSP surfaces. */
 	UPROPERTY(config, EditAnywhere, Category = "Place on Filters")
 	bool bBSP = true;
 
+	/** If the tool should place on foliage surfaces. */
 	UPROPERTY(config, EditAnywhere, Category = "Place on Filters")
 	bool bFoliage = false;
 
+	/** If the tool should place on translucent surfaces. */
 	UPROPERTY(config, EditAnywhere, Category = "Place on Filters")
 	bool bTranslucent = false;
 
-	// todo: asset data does not serialize out correctly
-	// maybe save soft object pointers, and convert in the UI to asset data?
+	// todo: palette save and load between sessions
 	//UPROPERTY(config)
 	TArray<FPaletteItem> PaletteItems;
 
diff --git a/Engine/Source/Editor/AssetPlacementEdMode/Private/Tools/PlacementBrushToolBase.cpp b/Engine/Source/Editor/AssetPlacementEdMode/Private/Tools/PlacementBrushToolBase.cpp
--- a/Engine/Source/Editor/AssetPlacementEdMode/Private/Tools/PlacementBrushToolBase.cpp
+++ b/Engine/Source/Editor/AssetPlacementEdMode/Private/Tools/PlacementBrushToolBase.cpp
@@ -85,70 +85,93 @@ bool UPlacementBrushToolBase::FindHitResultWithStartAndEndTraceVectors(FHitResul
 	return AInstancedFoliageActor::FoliageTrace(EditingWorld, OutHit, FDesiredFoliageInstance(TraceStart, TraceEnd, /* FoliageType= */ nullptr, TraceRadius), NAME_PlacementBrushTool, /* bReturnFaceIndex */ false, FilterFunc);
 }
 
-FTransform UPlacementBrushToolBase::GetFinalTransformFromHitLocationAndNormal(const FVector& InLocation, const FVector& InNormal)
+FTransform UPlacementBrushToolBase::GenerateTransformFromHitLocationAndNormal(const FVector& InLocation, const FVector& InNormal)
 {
-	FTransform FinalizedTransform(InLocation);
+	const UAssetPlacementSettings* PlacementSettings = GEditor->GetEditorSubsystem<UPlacementModeSubsystem>()->GetModeSettingsObject();
+	FTransform FinalizedTransform(GenerateRandomRotation(PlacementSettings), InLocation, GenerateRandomScale(PlacementSettings));
+	return FinalizeTransform(FinalizedTransform, InNormal, PlacementSettings);
+}
 
-	TWeakObjectPtr<const UAssetPlacementSettings> PlacementSettings = GEditor->GetEditorSubsystem<UPlacementModeSubsystem>()->GetModeSettingsObject();
-	if (!PlacementSettings.IsValid())
+FQuat UPlacementBrushToolBase::GenerateRandomRotation(const UAssetPlacementSettings* PlacementSettings)
+{
+	if (!PlacementSettings)
 	{
-		return FinalizedTransform;
+		return FQuat::Identity;
 	}
 
-	// Update the rotation if we need to
-	FRotator FinalRotation = GetFinalRotation(FinalizedTransform);
-	FinalizedTransform.SetRotation(FinalRotation.Quaternion());
-	FinalizedTransform.NormalizeRotation();
-
-	if (PlacementSettings->bUseRandomScale)
+	FRotator GeneratedRotation = FRotator::ZeroRotator;
+	auto GetRandomSignedValueInRange = [](const FFloatInterval& Range, bool bAllowSigned) -> float
 	{
-		// Until we have per object settings, just use a uniform scale, clamped from half to double size
-		FFloatInterval ScaleRange(0.5f, 2.0f);
-		FVector NewScale(ScaleRange.Interpolate(FMath::FRand()));
-		FinalizedTransform.SetScale3D(NewScale);
-	}
-
-	return FinalizedTransform;
-}
-
-FRotator UPlacementBrushToolBase::GetFinalRotation(const FTransform& InTransform)
-{
-	FRotator UpdatedRotation = InTransform.Rotator();
+		float Sign = (FMath::RandBool() && bAllowSigned) ? -1.0f : 1.0f;
+		return Range.Interpolate(FMath::FRand()) * Sign;
+	};
 
-	TWeakObjectPtr<const UAssetPlacementSettings> PlacementSettings = GEditor->GetEditorSubsystem<UPlacementModeSubsystem>()->GetModeSettingsObject();
-	if (!PlacementSettings.IsValid())
+	if (PlacementSettings->bUseRandomRotationX)
 	{
-		return UpdatedRotation;
+		GeneratedRotation.Roll = GetRandomSignedValueInRange(PlacementSettings->RandomRotationX, PlacementSettings->bAllowNegativeRotationX);
 	}
 
-	if (PlacementSettings->bUseRandomXRotation)
+	if (PlacementSettings->bUseRandomRotationY)
 	{
-		UpdatedRotation.Roll = PlacementSettings->RandomRotationX.Interpolate(FMath::FRand());
+		GeneratedRotation.Pitch = GetRandomSignedValueInRange(PlacementSettings->RandomRotationY, PlacementSettings->bAllowNegativeRotationY);
 	}
 
-	if (PlacementSettings->bUseRandomYRotation)
+	if (PlacementSettings->bUseRandomRotationZ)
 	{
-		UpdatedRotation.Pitch = PlacementSettings->RandomRotationY.Interpolate(FMath::FRand());
+		GeneratedRotation.Yaw = GetRandomSignedValueInRange(PlacementSettings->RandomRotationZ, PlacementSettings->bAllowNegativeRotationZ);
 	}
 
-	if (PlacementSettings->bUseRandomZRotation)
+	return GeneratedRotation.Quaternion();
+}
+
+FVector UPlacementBrushToolBase::GenerateRandomScale(const UAssetPlacementSettings* PlacementSettings)
+{
+	FVector GeneratedScale(1.0f);
+	if (!PlacementSettings || !PlacementSettings->bUseRandomScale)
 	{
-		UpdatedRotation.Yaw = PlacementSettings->RandomRotationZ.Interpolate(FMath::FRand());
+		return GeneratedScale;
 	}
 
-	// Cache this rotation before we start aligning to normal
-	LastGeneratedRotation = FQuat(UpdatedRotation);
+	auto GenerateRandomScaleComponent = [PlacementSettings]() -> float
+	{
+		float Sign = (FMath::RandBool() && PlacementSettings->bAllowNegativeScale) ? -1.0f : 1.0f;
+		return PlacementSettings->ScaleRange.Interpolate(FMath::FRand()) * Sign;
+	};
 
-	// Align to normal
-	if (PlacementSettings->bAlignToNormal)
+	switch (PlacementSettings->ScalingType)
 	{
-		UpdatedRotation = UpdateRotationAlignedToBrushNormal(PlacementSettings->AxisToAlignWithNormal, PlacementSettings->bInvertNormalAxis).Rotator();
+		case EFoliageScaling::Free:
+		{
+			GeneratedScale = FVector(GenerateRandomScaleComponent(), GenerateRandomScaleComponent(), GenerateRandomScaleComponent());
+			break;
+		}
+		case EFoliageScaling::Uniform:
+		{
+			float ScaleComponent = GenerateRandomScaleComponent();
+			GeneratedScale = FVector(ScaleComponent, ScaleComponent, ScaleComponent);
+			break;
+		}
+		case EFoliageScaling::LockXY:
+		{
+			GeneratedScale.Z = GenerateRandomScaleComponent();
+			break;
+		}
+		case EFoliageScaling::LockYZ:
+		{
+			GeneratedScale.X = GenerateRandomScaleComponent();
+			break;
+		}
+		case EFoliageScaling::LockXZ:
+		{
+			GeneratedScale.Y = GenerateRandomScaleComponent();
+			break;
+		}
 	}
 
-	return UpdatedRotation;
+	return GeneratedScale;
 }
 
-FQuat UPlacementBrushToolBase::UpdateRotationAlignedToBrushNormal(EAxis::Type InAlignmentAxis, bool bInvertAxis)
+FQuat UPlacementBrushToolBase::AlignRotationWithNormal(const FQuat& InRotation, const FVector& InNormal, EAxis::Type InAlignmentAxis, bool bInvertAxis)
 {
 	FVector AlignmentVector = FVector::ZeroVector;
 	switch (InAlignmentAxis)
@@ -172,7 +195,46 @@ FQuat UPlacementBrushToolBase::UpdateRotationAlignedToBrushNormal(EAxis::Type In
 		break;
 	}
 
-	return FindActorAlignmentRotation(LastGeneratedRotation, AlignmentVector, LastBrushStamp.WorldNormal, &LastAlignRotation);
+	return FindActorAlignmentRotation(InRotation, AlignmentVector, InNormal);
+}
+
+FTransform UPlacementBrushToolBase::FinalizeTransform(const FTransform& OriginalTransform, const FVector& InNormal, const UAssetPlacementSettings* PlacementSettings)
+{
+	if (!PlacementSettings)
+	{
+		return OriginalTransform;
+	}
+
+	FTransform FinalizedTransform(FQuat::Identity, OriginalTransform.GetTranslation(), OriginalTransform.GetScale3D());
+
+	// Add the world offset.
+	FVector WorldOffset = PlacementSettings->WorldLocationOffset;
+	if (PlacementSettings->bScaleWorldLocationOffset)
+	{
+		WorldOffset *= OriginalTransform.GetScale3D();
+	}
+	FinalizedTransform.AddToTranslation(WorldOffset);
+
+	// Align to normal.
+	FQuat AdjustedRotation(OriginalTransform.GetRotation());
+	if (PlacementSettings->bAlignToNormal)
+	{
+		AdjustedRotation = AlignRotationWithNormal(AdjustedRotation, InNormal, PlacementSettings->AxisToAlignWithNormal, PlacementSettings->bInvertNormalAxis);
+	}
+	AdjustedRotation.Normalize();
+	FinalizedTransform.SetRotation(AdjustedRotation);
+
+	// Add the relative offset.
+	{
+		FVector RelativeOffset = PlacementSettings->RelativeLocationOffset;
+		if (PlacementSettings->bScaleRelativeLocationOffset)
+		{
+			RelativeOffset *= OriginalTransform.GetScale3D();
+		}
+		FinalizedTransform.SetTranslation(FinalizedTransform.TransformPosition(RelativeOffset));
+	}
+
+	return FinalizedTransform;
 }
 
 TArray<FTypedElementHandle> UPlacementBrushToolBase::GetElementsInBrushRadius() const
diff --git a/Engine/Source/Editor/AssetPlacementEdMode/Private/Tools/PlacementBrushToolBase.h b/Engine/Source/Editor/AssetPlacementEdMode/Private/Tools/PlacementBrushToolBase.h
--- a/Engine/Source/Editor/AssetPlacementEdMode/Private/Tools/PlacementBrushToolBase.h
+++ b/Engine/Source/Editor/AssetPlacementEdMode/Private/Tools/PlacementBrushToolBase.h
@@ -38,15 +38,16 @@ public:
 protected:
 	virtual double EstimateMaximumTargetDimension() override;
 	bool FindHitResultWithStartAndEndTraceVectors(FHitResult& OutHit, const FVector& TraceStart, const FVector& TraceEnd, float TraceRadius = 0.0f);
-	FTransform GetFinalTransformFromHitLocationAndNormal(const FVector& InLocation, const FVector& InNormal);
+	static FTransform GenerateTransformFromHitLocationAndNormal(const FVector& InLocation, const FVector& InNormal);
 
 	// Gets a random rotation, and aligns it, based on the placement settings.
-	FRotator GetFinalRotation(const FTransform& InTransform);
+	static FQuat GenerateRandomRotation(const UAssetPlacementSettings* PlacementSettings);
+
+	// Generates a random scale, based on placement settings.
+	static FVector GenerateRandomScale(const UAssetPlacementSettings* PlacementSettings);
 
 	// Updates the last generated rotation to realign with the current brush position and normal.
-	FQuat UpdateRotationAlignedToBrushNormal(EAxis::Type InAlignmentAxis, bool bInvertAxis);
+	static FQuat AlignRotationWithNormal(const FQuat& InRotation, const FVector& InNormal, EAxis::Type InAlignmentAxis, bool bInvertAxis);
+	static FTransform FinalizeTransform(const FTransform& OriginalTransform, const FVector& InNormal, const UAssetPlacementSettings* PlacementSettings);
 	TArray<FTypedElementHandle> GetElementsInBrushRadius() const;
-
-	FQuat LastGeneratedRotation;
-	FQuat LastAlignRotation;
 };
diff --git a/Engine/Source/Editor/AssetPlacementEdMode/Private/Tools/PlacementPlaceSingleTool.cpp b/Engine/Source/Editor/AssetPlacementEdMode/Private/Tools/PlacementPlaceSingleTool.cpp
--- a/Engine/Source/Editor/AssetPlacementEdMode/Private/Tools/PlacementPlaceSingleTool.cpp
+++ b/Engine/Source/Editor/AssetPlacementEdMode/Private/Tools/PlacementPlaceSingleTool.cpp
@@ -27,6 +27,7 @@ UPlacementModePlaceSingleTool::UPlacementModePlaceSingleTool() = default;
 UPlacementModePlaceSingleTool::~UPlacementModePlaceSingleTool() = default;
 UPlacementModePlaceSingleTool::UPlacementModePlaceSingleTool(FVTableHelper& Helper)
 	: Super(Helper)
+	, LastGeneratedRotation(FQuat::Identity)
 {
 }
 
@@ -61,7 +62,7 @@ void UPlacementModePlaceSingleTool::OnClickPress(const FInputDeviceRay& PressPos
 	// Place the Preview data if we managed to get to a valid handled click.
 	FPlacementOptions PlacementOptions;
 	PlacementOptions.bPreferBatchPlacement = true;
-	PlacementOptions.bPreferInstancedPlacement = true;
+	PlacementOptions.bPreferInstancedPlacement = false;	// Off until ISM single selection support added.
 
 	UPlacementSubsystem* PlacementSubsystem = GEditor->GetEditorSubsystem<UPlacementSubsystem>();
 	if (PlacementSubsystem)
@@ -125,17 +126,20 @@ FInputRayHit UPlacementModePlaceSingleTool::BeginHoverSequenceHitTest(const FInp
 void UPlacementModePlaceSingleTool::GeneratePreviewPlacementData(const FInputDeviceRay& DevicePos)
 {
 	PreviewPlacementInfo.Reset();
+	LastGeneratedRotation = FQuat::Identity;
 
 	const UAssetPlacementSettings* PlacementSettings = GEditor->GetEditorSubsystem<UPlacementModeSubsystem>()->GetModeSettingsObject();
 	if (PlacementSettings && PlacementSettings->PaletteItems.Num())
 	{
 		int32 ItemIndex = FMath::RandHelper(PlacementSettings->PaletteItems.Num());
 		const FPaletteItem& ItemToPlace = PlacementSettings->PaletteItems[ItemIndex];
+		LastGeneratedRotation = GenerateRandomRotation(PlacementSettings);
+		FTransform TransformToUpdate(LastGeneratedRotation, LastBrushStamp.WorldPosition, GenerateRandomScale(PlacementSettings));
 
 		PreviewPlacementInfo = MakeUnique<FAssetPlacementInfo>();
 		PreviewPlacementInfo->AssetToPlace = ItemToPlace.AssetData;
 		PreviewPlacementInfo->FactoryOverride = ItemToPlace.FactoryOverride;
-		PreviewPlacementInfo->FinalizedTransform = GetFinalTransformFromHitLocationAndNormal(LastBrushStamp.HitResult.ImpactPoint, LastBrushStamp.HitResult.ImpactNormal);
+		PreviewPlacementInfo->FinalizedTransform = FinalizeTransform(TransformToUpdate, LastBrushStamp.WorldNormal, PlacementSettings);
 		PreviewPlacementInfo->PreferredLevel = GEditor->GetEditorWorldContext().World()->GetCurrentLevel();
 	}
 }
@@ -175,19 +179,11 @@ void UPlacementModePlaceSingleTool::UpdatePreviewElements(const FInputDeviceRay&
 		return;
 	}
 
-	// Update the location.
-	FTransform& FinalizedTransform = PreviewPlacementInfo->FinalizedTransform;
-	FinalizedTransform.SetLocation(LastBrushStamp.HitResult.ImpactPoint);
+	FTransform UpdatedTransform(LastGeneratedRotation, LastBrushStamp.WorldPosition, PreviewPlacementInfo->FinalizedTransform.GetScale3D());
+	UpdatedTransform = FinalizeTransform(UpdatedTransform, LastBrushStamp.WorldNormal, GEditor->GetEditorSubsystem<UPlacementModeSubsystem>()->GetModeSettingsObject());
 
-	// Update the rotation.
-	const UAssetPlacementSettings* PlacementSettings = GEditor->GetEditorSubsystem<UPlacementModeSubsystem>()->GetModeSettingsObject();
-	if (PlacementSettings && PlacementSettings->bAlignToNormal)
-	{
-		FQuat UpdatedRotation = UpdateRotationAlignedToBrushNormal(PlacementSettings->AxisToAlignWithNormal, PlacementSettings->bInvertNormalAxis);
-		FinalizedTransform.SetRotation(UpdatedRotation);
-	}
-
-	UpdateElementTransforms(PreviewElements, FinalizedTransform);
+	UpdateElementTransforms(PreviewElements, UpdatedTransform);
+	PreviewPlacementInfo->FinalizedTransform = UpdatedTransform;
 }
 
 void UPlacementModePlaceSingleTool::DestroyPreviewElements()
diff --git a/Engine/Source/Editor/AssetPlacementEdMode/Private/Tools/PlacementPlaceSingleTool.h b/Engine/Source/Editor/AssetPlacementEdMode/Private/Tools/PlacementPlaceSingleTool.h
--- a/Engine/Source/Editor/AssetPlacementEdMode/Private/Tools/PlacementPlaceSingleTool.h
+++ b/Engine/Source/Editor/AssetPlacementEdMode/Private/Tools/PlacementPlaceSingleTool.h
@@ -51,6 +51,7 @@ protected:
 	void NotifyMovementEnded(TArrayView<const FTypedElementHandle> InElements);
 	void SetupRightClickMouseBehavior();
 
+	FQuat LastGeneratedRotation;
 	TUniquePtr<FAssetPlacementInfo> PreviewPlacementInfo;
 	TArray<FTypedElementHandle> PreviewElements;
 	bool bIsTweaking;
diff --git a/Engine/Source/Editor/AssetPlacementEdMode/Private/Tools/PlacementPlaceTool.cpp b/Engine/Source/Editor/AssetPlacementEdMode/Private/Tools/PlacementPlaceTool.cpp
--- a/Engine/Source/Editor/AssetPlacementEdMode/Private/Tools/PlacementPlaceTool.cpp
+++ b/Engine/Source/Editor/AssetPlacementEdMode/Private/Tools/PlacementPlaceTool.cpp
@@ -43,7 +43,7 @@ void UPlacementModePlacementTool::OnTick(float DeltaTime)
 	{
 		FVector BrushLocation = LastBrushStamp.WorldPosition;
 
-		// Assume a default density of 100 *  whatever the user has selected as brush size.
+		// Assume a default density of 100 * whatever the user has selected as brush size.
 		float DefaultDensity = 100.0f * BrushProperties->BrushSize;
 		// This is the total set of instances disregarding parameters like slope, height or layer.
 		float DesiredInstanceCountFloat = DefaultDensity /** todo: individual item's Density*/ * BrushProperties->BrushStrength;
@@ -90,7 +90,7 @@ void UPlacementModePlacementTool::OnTick(float DeltaTime)
 			FAssetPlacementInfo NewInfo;
 			NewInfo.AssetToPlace = ItemToPlace.AssetData;
 			NewInfo.PreferredLevel = GEditor->GetEditorWorldContext().World()->GetCurrentLevel();
-			NewInfo.FinalizedTransform = GetFinalTransformFromHitLocationAndNormal(SpawnLocation, SpawnNormal);
+			NewInfo.FinalizedTransform = GenerateTransformFromHitLocationAndNormal(SpawnLocation, SpawnNormal);
 
 			DesiredInfoSpawnLocations.Emplace(NewInfo);
 		}
