Placement mode - replace palette using shift + drag&drop of assets from CB.

Also hooked up a clear all items from palette command, and removed some stale combo box functions.

diff --git a/Engine/Source/Editor/AssetPlacementEdMode/Private/AssetPlacementPaletteItem.cpp b/Engine/Source/Editor/AssetPlacementEdMode/Private/AssetPlacementPaletteItem.cpp
--- a/Engine/Source/Editor/AssetPlacementEdMode/Private/AssetPlacementPaletteItem.cpp
+++ b/Engine/Source/Editor/AssetPlacementEdMode/Private/AssetPlacementPaletteItem.cpp
@@ -189,8 +189,7 @@ EVisibility FAssetPlacementPaletteItemModel::GetTooltipVisibility() const
 
 EVisibility FAssetPlacementPaletteItemModel::GetTooltipThumbnailVisibility() const
 {
-	
-	return (AssetPalette.IsValid() && AssetPalette.Pin()->IsActiveViewMode(EAssetPlacementPaletteViewMode::Tree)) ? EVisibility::SelfHitTestInvisible : EVisibility::Collapsed;
+	return (AssetPalette.IsValid() && AssetPalette.Pin()->IsActiveViewMode(SAssetPlacementPalette::EViewMode::Tree)) ? EVisibility::SelfHitTestInvisible : EVisibility::Collapsed;
 }
 
 FText FAssetPlacementPaletteItemModel::GetSourceAssetTypeText() const
diff --git a/Engine/Source/Editor/AssetPlacementEdMode/Private/Modes/PlacementModeSubsystem.cpp b/Engine/Source/Editor/AssetPlacementEdMode/Private/Modes/PlacementModeSubsystem.cpp
--- a/Engine/Source/Editor/AssetPlacementEdMode/Private/Modes/PlacementModeSubsystem.cpp
+++ b/Engine/Source/Editor/AssetPlacementEdMode/Private/Modes/PlacementModeSubsystem.cpp
@@ -68,3 +68,8 @@ bool UPlacementModeSubsystem::AddPaletteItem(const FPaletteItem& InPaletteItem)
 
 	return false;
 }
+
+void UPlacementModeSubsystem::ClearPalette()
+{
+	ModeSettings->PaletteItems.Empty();
+}
diff --git a/Engine/Source/Editor/AssetPlacementEdMode/Private/Modes/PlacementModeSubsystem.h b/Engine/Source/Editor/AssetPlacementEdMode/Private/Modes/PlacementModeSubsystem.h
--- a/Engine/Source/Editor/AssetPlacementEdMode/Private/Modes/PlacementModeSubsystem.h
+++ b/Engine/Source/Editor/AssetPlacementEdMode/Private/Modes/PlacementModeSubsystem.h
@@ -39,6 +39,11 @@ public:
 	 */
 	bool AddPaletteItem(const FPaletteItem& InPaletteItem);
 
+	/**
+	 * Clears all items from the current palette.
+	 */
+	void ClearPalette();
+
 protected:
 	UPROPERTY()
 	TObjectPtr<UAssetPlacementSettings> ModeSettings;
diff --git a/Engine/Source/Editor/AssetPlacementEdMode/Private/SAssetPlacementPalette.cpp b/Engine/Source/Editor/AssetPlacementEdMode/Private/SAssetPlacementPalette.cpp
--- a/Engine/Source/Editor/AssetPlacementEdMode/Private/SAssetPlacementPalette.cpp
+++ b/Engine/Source/Editor/AssetPlacementEdMode/Private/SAssetPlacementPalette.cpp
@@ -272,7 +272,7 @@ void SAssetPlacementPalette::RefreshPalette()
 
 void SAssetPlacementPalette::RefreshActivePaletteViewWidget()
 {
-	if (ActiveViewMode == EAssetPlacementPaletteViewMode::Thumbnail)
+	if (ActiveViewMode == EViewMode::Thumbnail)
 	{
 		TileViewWidget->RequestListRefresh();
 	}
@@ -284,11 +284,6 @@ void SAssetPlacementPalette::RefreshActivePaletteViewWidget()
 
 void SAssetPlacementPalette::AddPlacementType(const FAssetData& AssetData)
 {
-	if (AddPlacementTypeCombo.IsValid())
-	{
-		AddPlacementTypeCombo->SetIsOpen(false);
-	}
-
 	if (!AssetData.IsValid())
 	{
 		return;
@@ -324,6 +319,22 @@ void SAssetPlacementPalette::AddPlacementType(const FAssetData& AssetData)
 	UAssetManager::GetStreamableManager().RequestAsyncLoad(AssetData.ToSoftObjectPath());
 
 	PaletteItems.Add(MakeShared<FAssetPlacementPaletteItemModel>(PlacementInfo, SharedThis(this), ThumbnailPool));
+}
+
+void SAssetPlacementPalette::ClearPalette()
+{
+	if (UPlacementModeSubsystem* ModeSubsystem = GEditor->GetEditorSubsystem<UPlacementModeSubsystem>())
+	{
+		ModeSubsystem->ClearPalette();
+	}
+
+	PaletteItems.Empty();
+	FilteredItems.Empty();
+}
+
+void SAssetPlacementPalette::OnClearPalette()
+{
+	ClearPalette();
 	UpdatePalette(true);
 }
 
@@ -362,7 +373,7 @@ TSharedRef<SWidgetSwitcher> SAssetPlacementPalette::CreatePaletteViews()
 	SAssignNew(WidgetSwitcher, SWidgetSwitcher);
 
 	// Thumbnail View
-	WidgetSwitcher->AddSlot((int32)EAssetPlacementPaletteViewMode::Thumbnail)
+	WidgetSwitcher->AddSlot((uint8)EViewMode::Thumbnail)
 	[
 		SNew(SScrollBorder, TileViewWidget.ToSharedRef())
 		.Content()
@@ -372,7 +383,7 @@ TSharedRef<SWidgetSwitcher> SAssetPlacementPalette::CreatePaletteViews()
 	];
 
 	// Tree View
-	WidgetSwitcher->AddSlot((int32)EAssetPlacementPaletteViewMode::Tree)
+	WidgetSwitcher->AddSlot((uint8)EViewMode::Tree)
 	[
 		SNew(SScrollBorder, TreeViewWidget.ToSharedRef())
 		.Style(&FAssetPlacementEdModeStyle::Get().GetWidgetStyle<FScrollBorderStyle>("FoliageEditMode.TreeView.ScrollBorder"))
@@ -382,7 +393,7 @@ TSharedRef<SWidgetSwitcher> SAssetPlacementPalette::CreatePaletteViews()
 		]
 	];
 
-	WidgetSwitcher->SetActiveWidgetIndex((int32)ActiveViewMode);
+	WidgetSwitcher->SetActiveWidgetIndex((uint8)ActiveViewMode);
 
 	return WidgetSwitcher.ToSharedRef();
 }
@@ -399,17 +410,6 @@ void SAssetPlacementPalette::OnSearchTextChanged(const FText& InFilterText)
 	UpdatePalette();
 }
 
-TSharedRef<SWidget> SAssetPlacementPalette::GetAddPlacementTypePicker()
-{
-	return PropertyCustomizationHelpers::MakeAssetPickerWithMenu(FAssetData(),
-		false,
-		TArray<const UClass*>({UObject::StaticClass()}),
-		TArray<UFactory*>(),
-		FOnShouldFilterAsset::CreateSP(this, &SAssetPlacementPalette::ShouldFilterAsset),
-		FOnAssetSelected::CreateSP(this, &SAssetPlacementPalette::AddPlacementType),
-		FSimpleDelegate());
-}
-
 bool SAssetPlacementPalette::ShouldFilterAsset(const FAssetData& InAssetData)
 {
 	UClass* Class = InAssetData.GetClass();
@@ -427,18 +427,18 @@ bool SAssetPlacementPalette::ShouldFilterAsset(const FAssetData& InAssetData)
 	return true;
 }
 
-void SAssetPlacementPalette::SetViewMode(EAssetPlacementPaletteViewMode NewViewMode)
+void SAssetPlacementPalette::SetViewMode(EViewMode NewViewMode)
 {
 	if (ActiveViewMode != NewViewMode)
 	{
 		ActiveViewMode = NewViewMode;
-		WidgetSwitcher->SetActiveWidgetIndex((int32)ActiveViewMode);
+		WidgetSwitcher->SetActiveWidgetIndex((uint8)ActiveViewMode);
 		
 		RefreshActivePaletteViewWidget();
 	}
 }
 
-bool SAssetPlacementPalette::IsActiveViewMode(EAssetPlacementPaletteViewMode ViewMode) const
+bool SAssetPlacementPalette::IsActiveViewMode(EViewMode ViewMode) const
 {
 	return ActiveViewMode == ViewMode;
 }
@@ -469,9 +469,9 @@ TSharedRef<SWidget> SAssetPlacementPalette::GetViewOptionsMenuContent()
 			LOCTEXT("ThumbnailView_ToolTip", "Display thumbnails for each Placement type in the palette."),
 			FSlateIcon(),
 			FUIAction(
-				FExecuteAction::CreateSP(this, &SAssetPlacementPalette::SetViewMode, EAssetPlacementPaletteViewMode::Thumbnail),
+				FExecuteAction::CreateSP(this, &SAssetPlacementPalette::SetViewMode, EViewMode::Thumbnail),
 				FCanExecuteAction(),
-				FIsActionChecked::CreateSP(this, &SAssetPlacementPalette::IsActiveViewMode, EAssetPlacementPaletteViewMode::Thumbnail)
+				FIsActionChecked::CreateSP(this, &SAssetPlacementPalette::IsActiveViewMode, EViewMode::Thumbnail)
 				),
 			NAME_None,
 			EUserInterfaceActionType::ToggleButton
@@ -482,9 +482,9 @@ TSharedRef<SWidget> SAssetPlacementPalette::GetViewOptionsMenuContent()
 			LOCTEXT("ListView_ToolTip", "Display Placement types in the palette as a list."),
 			FSlateIcon(),
 			FUIAction(
-				FExecuteAction::CreateSP(this, &SAssetPlacementPalette::SetViewMode, EAssetPlacementPaletteViewMode::Tree),
+				FExecuteAction::CreateSP(this, &SAssetPlacementPalette::SetViewMode, EViewMode::Tree),
 				FCanExecuteAction(),
-				FIsActionChecked::CreateSP(this, &SAssetPlacementPalette::IsActiveViewMode, EAssetPlacementPaletteViewMode::Tree)
+				FIsActionChecked::CreateSP(this, &SAssetPlacementPalette::IsActiveViewMode, EViewMode::Tree)
 				),
 			NAME_None,
 			EUserInterfaceActionType::ToggleButton
@@ -502,7 +502,7 @@ TSharedRef<SWidget> SAssetPlacementPalette::GetViewOptionsMenuContent()
 				FExecuteAction::CreateSP(this, &SAssetPlacementPalette::ToggleShowTooltips),
 				FCanExecuteAction(),
 				FIsActionChecked::CreateSP(this, &SAssetPlacementPalette::ShouldShowTooltips),
-				FIsActionButtonVisible::CreateSP(this, &SAssetPlacementPalette::IsActiveViewMode, EAssetPlacementPaletteViewMode::Tree)
+				FIsActionButtonVisible::CreateSP(this, &SAssetPlacementPalette::IsActiveViewMode, EViewMode::Tree)
 				),
 			NAME_None,
 			EUserInterfaceActionType::ToggleButton
@@ -535,11 +535,11 @@ TSharedRef<SWidget> SAssetPlacementPalette::GetViewOptionsMenuContent()
 
 TSharedPtr<SListView<FPlacementPaletteItemModelPtr>> SAssetPlacementPalette::GetActiveViewWidget() const
 {
-	if (ActiveViewMode == EAssetPlacementPaletteViewMode::Thumbnail)
+	if (ActiveViewMode == EViewMode::Thumbnail)
 	{
 		return TileViewWidget;
 	}
-	else if (ActiveViewMode == EAssetPlacementPaletteViewMode::Tree)
+	else if (ActiveViewMode == EViewMode::Tree)
 	{
 		return TreeViewWidget;
 	}
@@ -579,21 +579,36 @@ FReply SAssetPlacementPalette::HandlePlacementDropped(const FGeometry& DropZoneG
 	TArray<FAssetData> DroppedAssetData = AssetUtil::ExtractAssetDataFromDrag(DragDropEvent);
 	if (DroppedAssetData.Num() > 0)
 	{
-		// Treat the entire drop as a transaction (in case multiples types are being added)
-		const FScopedTransaction Transaction(NSLOCTEXT("UnrealEd", "PlacementMode_DragDropTypesTransaction", "Drag-drop Placement"));
+		if (DragDropEvent.IsShiftDown())
+		{
+			ClearPalette();
+		}
 
 		for (auto& AssetData : DroppedAssetData)
 		{
 			AddPlacementType(AssetData);
 		}
+
+		UpdatePalette(true);
 	}
 
 	return FReply::Handled();
 }
 
+bool SAssetPlacementPalette::HasAnyItemInPalette() const
+{
+	return PaletteItems.Num() > 0;
+}
+
 TSharedPtr<SWidget> SAssetPlacementPalette::ConstructPlacementTypeContextMenu()
 {
 	FMenuBuilder MenuBuilder(true, nullptr);
+	MenuBuilder.AddMenuEntry(
+		LOCTEXT("Palette_Clear", "Clear Palette"),
+		LOCTEXT("Palette_ClearDesc", "Removes all items from the palette."),
+		FSlateIcon(),
+		FUIAction(FExecuteAction::CreateSP(this, &SAssetPlacementPalette::OnClearPalette), FCanExecuteAction::CreateSP(this, &SAssetPlacementPalette::HasAnyItemInPalette))
+	);
 	return MenuBuilder.MakeWidget();
 }
 
@@ -642,7 +657,7 @@ void SAssetPlacementPalette::SetThumbnailScale(float InScale)
 
 EVisibility SAssetPlacementPalette::GetThumbnailScaleSliderVisibility() const
 {
-	return (ActiveViewMode == EAssetPlacementPaletteViewMode::Thumbnail) ? EVisibility::Visible : EVisibility::Collapsed;
+	return (ActiveViewMode == EViewMode::Thumbnail) ? EVisibility::Visible : EVisibility::Collapsed;
 }
 
 // TREE VIEW
diff --git a/Engine/Source/Editor/AssetPlacementEdMode/Private/SAssetPlacementPalette.h b/Engine/Source/Editor/AssetPlacementEdMode/Private/SAssetPlacementPalette.h
--- a/Engine/Source/Editor/AssetPlacementEdMode/Private/SAssetPlacementPalette.h
+++ b/Engine/Source/Editor/AssetPlacementEdMode/Private/SAssetPlacementPalette.h
@@ -31,17 +31,17 @@ typedef TSharedPtr<FAssetPlacementPaletteItemModel> FPlacementPaletteItemModelPt
 typedef STreeView<FPlacementPaletteItemModelPtr> SPlacementTypeTreeView;
 typedef STileView<FPlacementPaletteItemModelPtr> SPlacementTypeTileView;
 
-/** View modes supported by the palette */
-enum class EAssetPlacementPaletteViewMode : int32
-{
-	Thumbnail,
-	Tree
-};
-
 /** The palette of Placement types available for use by the Placement edit mode */
 class SAssetPlacementPalette : public SCompoundWidget
 {
 public:
+	/** View modes supported by the palette */
+	enum class EViewMode : uint8
+	{
+		Thumbnail,
+		Tree
+	};
+
 	SLATE_BEGIN_ARGS(SAssetPlacementPalette) {}
 		SLATE_ARGUMENT(TSharedPtr<IPropertyHandle>, PalettePropertyHandle)
 	SLATE_END_ARGS()
@@ -59,7 +59,7 @@ public:
 	void RefreshPalette();
 
 	/** @return True if the given view mode is the active view mode */
-	bool IsActiveViewMode(EAssetPlacementPaletteViewMode ViewMode) const;
+	bool IsActiveViewMode(EViewMode ViewMode) const;
 
 	/** @return True if tooltips should be shown when hovering over Placement type items in the palette */
 	bool ShouldShowTooltips() const;
@@ -67,10 +67,15 @@ public:
 	/** @return The current search filter text */
 	FText GetSearchText() const;
 
+private:
 	/** Adds the Placement type asset to the instanced Placement actor's list of types. */
 	void AddPlacementType(const FAssetData& AssetData);
 
-private:	// GENERAL
+	/**
+	 * Removes all items from the palette.
+	 */
+	void ClearPalette();
+	void OnClearPalette();
 
 	/** Refreshes the active palette view widget */
 	void RefreshActivePaletteViewWidget();
@@ -84,16 +89,13 @@ private:	// GENERAL
 	/** Handles changes to the search filter text */
 	void OnSearchTextChanged(const FText& InFilterText);
 
-	/** Gets the asset picker for adding a Placement type. */
-	TSharedRef<SWidget> GetAddPlacementTypePicker();
-
 	bool ShouldFilterAsset(const FAssetData& InAssetData);
 
 	/** Gets the visibility of the Add Placement Type text in the header row button */
 	EVisibility GetAddPlacementTypeButtonTextVisibility() const;
 
 	/** Sets the view mode of the palette */
-	void SetViewMode(EAssetPlacementPaletteViewMode NewViewMode);
+	void SetViewMode(EViewMode NewViewMode);
 
 	/** Sets whether to show tooltips when hovering over Placement type items in the palette */
 	void ToggleShowTooltips();
@@ -118,6 +120,9 @@ private:	// GENERAL
 	/** Handles dropping of a mesh or Placement type into the palette */
 	FReply HandlePlacementDropped(const FGeometry& DropZoneGeometry, const FDragDropEvent& DragDropEvent);
 
+	/** @returns true if there are any items in the palette. */
+	bool HasAnyItemInPalette() const;
+
 private:	// CONTEXT MENU
 
 	/** @return the SWidget containing the context menu */
@@ -180,9 +185,6 @@ private:
 	/** Switches between the thumbnail and tree views */
 	TSharedPtr<class SWidgetSwitcher> WidgetSwitcher;
 
-	/** The Add Placement Type combo button */
-	TSharedPtr<class SComboButton> AddPlacementTypeCombo;
-
 	/** The header row of the Placement mesh tree */
 	TSharedPtr<class SHeaderRow> TreeViewHeaderRow;
 
@@ -210,7 +212,7 @@ private:
 	bool bShowFullTooltips = true;
 	bool bIsRebuildTimerRegistered = true;
 	bool bIsRefreshTimerRegistered = true;
-	EAssetPlacementPaletteViewMode ActiveViewMode;
+	EViewMode ActiveViewMode;
 	EColumnSortMode::Type ActiveSortOrder = EColumnSortMode::Type::Ascending;
 
 	float PaletteThumbnailScale = .3f;
