Dragging after placement now allows for scale and rotate.

Added an option to prevent auto selection after placement for more rapid placing.

diff --git a/Engine/Source/Editor/AssetPlacementEdMode/Private/AssetPlacementSettings.cpp b/Engine/Source/Editor/AssetPlacementEdMode/Private/AssetPlacementSettings.cpp
--- a/Engine/Source/Editor/AssetPlacementEdMode/Private/AssetPlacementSettings.cpp
+++ b/Engine/Source/Editor/AssetPlacementEdMode/Private/AssetPlacementSettings.cpp
@@ -10,7 +10,7 @@ bool UAssetPlacementSettings::CanEditChange(const FProperty* InProperty) const
 	}
 
 	const FName PropertyName = InProperty->GetFName();
-	if (PropertyName == GET_MEMBER_NAME_CHECKED(UAssetPlacementSettings, ScaleRange) || PropertyName == GET_MEMBER_NAME_CHECKED(UAssetPlacementSettings, bAllowNegativeScale))
+	if (PropertyName == GET_MEMBER_NAME_CHECKED(UAssetPlacementSettings, bAllowNegativeScale))
 	{
 		return bUseRandomScale;
 	}
diff --git a/Engine/Source/Editor/AssetPlacementEdMode/Private/AssetPlacementSettings.h b/Engine/Source/Editor/AssetPlacementEdMode/Private/AssetPlacementSettings.h
--- a/Engine/Source/Editor/AssetPlacementEdMode/Private/AssetPlacementSettings.h
+++ b/Engine/Source/Editor/AssetPlacementEdMode/Private/AssetPlacementSettings.h
@@ -76,18 +76,18 @@ public:
 	UPROPERTY(config, EditAnywhere, Category = "Rotation")
 	bool bAllowNegativeRotationZ = true;
 
-	/** Set to true if the tool should randomly scale any generated assets */
-	UPROPERTY(config, EditAnywhere, Category = "Scale", meta = (InlineEditConditionToggle))
-	bool bUseRandomScale = true;
-
-	/** Which axes any randomly generated scale should be applied to. Any locked axes will be scaled to 1.0f. */
-	UPROPERTY(config, EditAnywhere, Category = "Scale", meta = (DisplayName="Use Random Scale", EditCondition = "bUseRandomScale"))
+	/** Which axes scale should be applied to. Any locked axes will be scaled to 1.0f. */
+	UPROPERTY(config, EditAnywhere, Category = "Scale")
 	EFoliageScaling ScalingType = EFoliageScaling::Uniform;
 	
-	/** The min and max at which a random scale is generated. */
+	/** The min and max at which a random scale is generated, or a scale is applied. */
 	UPROPERTY(config, EditAnywhere, Category = "Scale", meta = (UIMin = ".01"))
 	FFloatInterval ScaleRange = FFloatInterval(.5f, 2.0f);
 
+	/** Set to true if the tool should randomly scale any generated assets */
+	UPROPERTY(config, EditAnywhere, Category = "Scale", meta = (InlineEditConditionToggle))
+	bool bUseRandomScale = true;
+
 	/** Set to true if the tool should generate a negatively signed scale. */
 	UPROPERTY(config, EditAnywhere, AdvancedDisplay, Category = "Scale")
 	bool bAllowNegativeScale = false;
diff --git a/Engine/Source/Editor/AssetPlacementEdMode/Private/Tools/PlacementBrushToolBase.cpp b/Engine/Source/Editor/AssetPlacementEdMode/Private/Tools/PlacementBrushToolBase.cpp
--- a/Engine/Source/Editor/AssetPlacementEdMode/Private/Tools/PlacementBrushToolBase.cpp
+++ b/Engine/Source/Editor/AssetPlacementEdMode/Private/Tools/PlacementBrushToolBase.cpp
@@ -174,7 +174,7 @@ FVector UPlacementBrushToolBase::GenerateRandomScale(const UAssetPlacementSettin
 
 FQuat UPlacementBrushToolBase::AlignRotationWithNormal(const FQuat& InRotation, const FVector& InNormal, EAxis::Type InAlignmentAxis, bool bInvertAxis)
 {
-	FVector AlignmentVector = FVector::ZeroVector;
+	FVector AlignmentVector = FVector::UpVector;
 	switch (InAlignmentAxis)
 	{
 		case EAxis::Type::X:
diff --git a/Engine/Source/Editor/AssetPlacementEdMode/Private/Tools/PlacementPlaceSingleTool.cpp b/Engine/Source/Editor/AssetPlacementEdMode/Private/Tools/PlacementPlaceSingleTool.cpp
--- a/Engine/Source/Editor/AssetPlacementEdMode/Private/Tools/PlacementPlaceSingleTool.cpp
+++ b/Engine/Source/Editor/AssetPlacementEdMode/Private/Tools/PlacementPlaceSingleTool.cpp
@@ -6,7 +6,10 @@
 #include "Editor.h"
 #include "InteractiveToolManager.h"
 #include "ScopedTransaction.h"
+#include "ToolDataVisualizer.h"
 #include "BaseBehaviors/ClickDragBehavior.h"
+#include "BaseGizmos/BrushStampIndicator.h"
+#include "BaseGizmos/GizmoRenderingUtil.h"
 #include "Editor/EditorEngine.h"
 #include "Elements/Framework/TypedElementHandle.h"
 #include "Elements/Framework/TypedElementRegistry.h"
@@ -28,7 +31,6 @@ UPlacementModePlaceSingleTool::UPlacementModePlaceSingleTool() = default;
 UPlacementModePlaceSingleTool::~UPlacementModePlaceSingleTool() = default;
 UPlacementModePlaceSingleTool::UPlacementModePlaceSingleTool(FVTableHelper& Helper)
 	: Super(Helper)
-	, LastGeneratedRotation(FQuat::Identity)
 {
 }
 
@@ -39,6 +41,10 @@ void UPlacementModePlaceSingleTool::Setup()
 	SetupRightClickMouseBehavior();
 
 	bIsTweaking = false;
+
+	SinglePlaceSettings = NewObject<UPlacementModePlaceSingleToolSettings>(this);
+	SinglePlaceSettings->LoadConfig();
+	AddToolPropertySource(SinglePlaceSettings);
 }
 
 void UPlacementModePlaceSingleTool::Shutdown(EToolShutdownType ShutdownType)
@@ -50,16 +56,15 @@ void UPlacementModePlaceSingleTool::Shutdown(EToolShutdownType ShutdownType)
 	ExitTweakState(bClearSelectionSet);
 
 	Super::Shutdown(ShutdownType);
+
+	SinglePlaceSettings->SaveConfig();
+	SinglePlaceSettings = nullptr;
 }
 
 void UPlacementModePlaceSingleTool::OnClickPress(const FInputDeviceRay& PressPos)
 {
 	Super::OnClickPress(PressPos);
 
-	// Update the preview element one final time to be sure the transform is updated.
-	UpdatePreviewElements(PressPos);
-	DestroyPreviewElements();
-
 	// Place the Preview data if we managed to get to a valid handled click.
 	FPlacementOptions PlacementOptions;
 	PlacementOptions.bPreferBatchPlacement = true;
@@ -68,12 +73,102 @@ void UPlacementModePlaceSingleTool::OnClickPress(const FInputDeviceRay& PressPos
 	UPlacementSubsystem* PlacementSubsystem = GEditor->GetEditorSubsystem<UPlacementSubsystem>();
 	if (PlacementSubsystem)
 	{
-		// Update the level, just in case the preview moved us out of the current one.
-		PreviewPlacementInfo->PreferredLevel = GEditor->GetEditorWorldContext().World()->GetCurrentLevel();
+		FAssetPlacementInfo FinalizedPlacementInfo = *PlacementInfo;
+		FinalizedPlacementInfo.FinalizedTransform = FinalizeTransform(
+			FTransform(PlacementInfo->FinalizedTransform.GetRotation(), LastBrushStamp.WorldPosition, PlacementInfo->FinalizedTransform.GetScale3D()),
+			LastBrushStamp.WorldNormal,
+			GEditor->GetEditorSubsystem<UPlacementModeSubsystem>()->GetModeSettingsObject());
+		FinalizedPlacementInfo.PreferredLevel = GEditor->GetEditorWorldContext().World()->GetCurrentLevel();
+
+		GetToolManager()->BeginUndoTransaction(NSLOCTEXT("PlacementMode", "SinglePlaceAsset", "Place Single Asset"));
+		PlacedElements = PlacementSubsystem->PlaceAsset(FinalizedPlacementInfo, PlacementOptions);
+		NotifyMovementStarted(PlacedElements);
+	}
+}
 
-		FScopedTransaction Transaction(NSLOCTEXT("PlacementMode", "SinglePlaceAsset", "Place Single Asset"));
-		EnterTweakState(PlacementSubsystem->PlaceAsset(*PreviewPlacementInfo, PlacementOptions));
+void UPlacementModePlaceSingleTool::OnClickDrag(const FInputDeviceRay& DragPos)
+{
+	FVector TraceStartLocation = DragPos.WorldRay.Origin;
+	FVector TraceDirection = DragPos.WorldRay.Direction;
+	FVector TraceEndLocation = TraceStartLocation + (TraceDirection * HALF_WORLD_MAX);
+	FVector TraceIntersectionXY = FMath::LinePlaneIntersection(TraceStartLocation, TraceEndLocation, FPlane(LastBrushStamp.WorldPosition, LastBrushStamp.WorldNormal));
+
+	FVector CursorDirection;
+	float CursorDistance;
+	FVector MouseDelta = LastBrushStamp.WorldPosition - TraceIntersectionXY;
+	MouseDelta.ToDirectionAndLength(CursorDirection, CursorDistance);
+	if (CursorDirection.IsNearlyZero())
+	{
+		return;
+	}
+
+	const UAssetPlacementSettings* PlacementSettings = GEditor->GetEditorSubsystem<UPlacementModeSubsystem>()->GetModeSettingsObject();
+
+	// Update rotation based on mouse position
+	FTransform UpdatedTransform = FinalizeTransform(
+		FTransform(FRotationMatrix::MakeFromXZ(CursorDirection, PlacementInfo->FinalizedTransform.GetRotation().GetUpVector()).ToQuat(), LastBrushStamp.WorldPosition, PlacementInfo->FinalizedTransform.GetScale3D()),
+		LastBrushStamp.WorldNormal,
+		PlacementSettings);
+
+	// Update scale based on mouse position
+	FVector UpdatedScale = UpdatedTransform.GetScale3D();
+	if (PlacementSettings && (SinglePlaceSettings->ScalingType != EPlacementScaleToCursorType::None))
+	{
+		auto UpdateComponent = [PlacementSettings, CursorDistance, this](float InComponent) -> float
+		{
+			float Sign = 1.0f;
+			if (InComponent < 0.0f)
+			{
+				Sign = -1.0f;
+			}
+
+			FFloatInterval ScaleRange(FMath::Abs(InComponent), PlacementSettings->ScaleRange.Max);
+			float NewComponent = ScaleRange.Interpolate(FMath::Min(1.0f, (CursorDistance / BrushStampIndicator->BrushRadius)));
+			return NewComponent * Sign;
+		};
+
+		switch (PlacementSettings->ScalingType)
+		{
+			case EFoliageScaling::LockXY:
+			{
+				UpdatedScale.Z = UpdateComponent(UpdatedScale.Z);
+				break;
+			}
+			case EFoliageScaling::LockYZ:
+			{
+				UpdatedScale.X = UpdateComponent(UpdatedScale.X);
+				break;
+			}
+			case EFoliageScaling::LockXZ:
+			{
+				UpdatedScale.Y = UpdateComponent(UpdatedScale.Y);
+				break;
+			}
+			default:
+			{
+				UpdatedScale.X = UpdateComponent(UpdatedScale.X);
+				UpdatedScale.Y = UpdateComponent(UpdatedScale.Y);
+				UpdatedScale.Z = UpdateComponent(UpdatedScale.Z);
+				break;
+			}
+		}
 	}
+	UpdatedTransform.SetScale3D(UpdatedScale);
+
+	// Use the drag position and settings to update the scale and rotation of the placed elements.
+	UpdateElementTransforms(PlacedElements, UpdatedTransform, false);
+}
+
+void UPlacementModePlaceSingleTool::OnClickRelease(const FInputDeviceRay& ReleasePos)
+{
+	Super::OnClickRelease(ReleasePos);
+
+	NotifyMovementEnded(PlacedElements);
+	EnterTweakState(PlacedElements);
+	GetToolManager()->EndUndoTransaction();
+
+	ShutdownBrushStampIndicator();
+	PlacementInfo.Reset();
 }
 
 void UPlacementModePlaceSingleTool::OnBeginHover(const FInputDeviceRay& DevicePos)
@@ -81,6 +176,7 @@ void UPlacementModePlaceSingleTool::OnBeginHover(const FInputDeviceRay& DevicePo
 	Super::OnBeginHover(DevicePos);
 
 	// Always regenerate the placement data when a hover sequence begins
+	PlacementInfo.Reset();
 	CreatePreviewElements(DevicePos);
 }
 
@@ -124,45 +220,55 @@ FInputRayHit UPlacementModePlaceSingleTool::BeginHoverSequenceHitTest(const FInp
 	return FInputRayHit();
 }
 
-void UPlacementModePlaceSingleTool::GeneratePreviewPlacementData(const FInputDeviceRay& DevicePos)
+void UPlacementModePlaceSingleTool::Render(IToolsContextRenderAPI* RenderAPI)
 {
-	PreviewPlacementInfo.Reset();
-	LastGeneratedRotation = FQuat::Identity;
+	// Transform the brush radius to standard pixel size
+	float BrushRadiusScale = GizmoRenderingUtil::CalculateLocalPixelToWorldScale(RenderAPI->GetSceneView(), LastBrushStamp.WorldPosition);
+	LastBrushStamp.Radius = 100.0f * BrushRadiusScale;
 
+	Super::Render(RenderAPI);
+}
+
+void UPlacementModePlaceSingleTool::GeneratePlacementData(const FInputDeviceRay& DevicePos)
+{
 	const UAssetPlacementSettings* PlacementSettings = GEditor->GetEditorSubsystem<UPlacementModeSubsystem>()->GetModeSettingsObject();
 	if (PlacementSettings && PlacementSettings->PaletteItems.Num())
 	{
 		int32 ItemIndex = FMath::RandHelper(PlacementSettings->PaletteItems.Num());
 		const FPaletteItem& ItemToPlace = PlacementSettings->PaletteItems[ItemIndex];
-		LastGeneratedRotation = GenerateRandomRotation(PlacementSettings);
-		FTransform TransformToUpdate(LastGeneratedRotation, LastBrushStamp.WorldPosition, GenerateRandomScale(PlacementSettings));
-
-		PreviewPlacementInfo = MakeUnique<FAssetPlacementInfo>();
-		PreviewPlacementInfo->AssetToPlace = ItemToPlace.AssetData;
-		PreviewPlacementInfo->FactoryOverride = ItemToPlace.FactoryOverride;
-		PreviewPlacementInfo->FinalizedTransform = FinalizeTransform(TransformToUpdate, LastBrushStamp.WorldNormal, PlacementSettings);
-		PreviewPlacementInfo->PreferredLevel = GEditor->GetEditorWorldContext().World()->GetCurrentLevel();
+		FTransform TransformToUpdate(GenerateRandomRotation(PlacementSettings), LastBrushStamp.WorldPosition, GenerateRandomScale(PlacementSettings));
+
+		PlacementInfo = MakeUnique<FAssetPlacementInfo>();
+		PlacementInfo->AssetToPlace = ItemToPlace.AssetData;
+		PlacementInfo->FactoryOverride = ItemToPlace.FactoryOverride;
+		PlacementInfo->FinalizedTransform = FinalizeTransform(TransformToUpdate, LastBrushStamp.WorldNormal, PlacementSettings);
+		PlacementInfo->PreferredLevel = GEditor->GetEditorWorldContext().World()->GetCurrentLevel();
 	}
 }
 
 void UPlacementModePlaceSingleTool::CreatePreviewElements(const FInputDeviceRay& DevicePos)
 {
+	SetupBrushStampIndicator();
+
 	// Place the preview elements from our stored info
-	if (!PreviewPlacementInfo)
+	if (!PlacementInfo)
 	{
-		GeneratePreviewPlacementData(DevicePos);
+		GeneratePlacementData(DevicePos);
 	}
 
 	if (UPlacementSubsystem* PlacementSubsystem = GEditor->GetEditorSubsystem<UPlacementSubsystem>())
 	{
 		FPlacementOptions PlacementOptions;
 		PlacementOptions.bIsCreatingPreviewElements = true;
+		FAssetPlacementInfo InfoToPlace = *PlacementInfo;
+		InfoToPlace.FinalizedTransform = FinalizeTransform(
+			FTransform(PlacementInfo->FinalizedTransform.GetRotation(), LastBrushStamp.WorldPosition, PlacementInfo->FinalizedTransform.GetScale3D()),
+			LastBrushStamp.WorldNormal,
+			GEditor->GetEditorSubsystem<UPlacementModeSubsystem>()->GetModeSettingsObject());
 
-		PreviewElements = PlacementSubsystem->PlaceAsset(*PreviewPlacementInfo, PlacementOptions);
+		PreviewElements = PlacementSubsystem->PlaceAsset(InfoToPlace, PlacementOptions);
 	}
 
-	SetupBrushStampIndicator();
-
 	NotifyMovementStarted(PreviewElements);
 }
 
@@ -180,23 +286,21 @@ void UPlacementModePlaceSingleTool::UpdatePreviewElements(const FInputDeviceRay&
 		return;
 	}
 
-	FTransform UpdatedTransform(LastGeneratedRotation, LastBrushStamp.WorldPosition, PreviewPlacementInfo->FinalizedTransform.GetScale3D());
+	FTransform UpdatedTransform(PlacementInfo->FinalizedTransform.GetRotation(), LastBrushStamp.WorldPosition, PlacementInfo->FinalizedTransform.GetScale3D());
 	UpdatedTransform = FinalizeTransform(UpdatedTransform, LastBrushStamp.WorldNormal, GEditor->GetEditorSubsystem<UPlacementModeSubsystem>()->GetModeSettingsObject());
 
 	UpdateElementTransforms(PreviewElements, UpdatedTransform);
-	PreviewPlacementInfo->FinalizedTransform = UpdatedTransform;
 }
 
 void UPlacementModePlaceSingleTool::DestroyPreviewElements()
 {
 	NotifyMovementEnded(PreviewElements);
 	PreviewElements.Empty();
-	ShutdownBrushStampIndicator();
 }
 
 void UPlacementModePlaceSingleTool::EnterTweakState(TArrayView<const FTypedElementHandle> InElementHandles)
 {
-	if (InElementHandles.Num() == 0)
+	if (InElementHandles.Num() == 0 || !SinglePlaceSettings->bSelectAfterPlacing)
 	{
 		return;
 	}
@@ -217,6 +321,7 @@ void UPlacementModePlaceSingleTool::ExitTweakState(bool bClearSelectionSet)
 	}
 
 	bIsTweaking = false;
+	PlacedElements.Empty();
 }
 
 void UPlacementModePlaceSingleTool::UpdateElementTransforms(TArrayView<const FTypedElementHandle> InElements, const FTransform& InTransform, bool bLocalTransform)
@@ -268,7 +373,7 @@ void UPlacementModePlaceSingleTool::SetupRightClickMouseBehavior()
 	{
 		constexpr bool bClearSelection = true;
 		ExitTweakState(bClearSelection);
-		PreviewPlacementInfo.Reset();
+		PlacementInfo.Reset();
 	};
 	RightMouseBehavior->SetDefaultPriority(FInputCapturePriority(-1));
 	RightMouseBehavior->SetUseRightMouseButton();
diff --git a/Engine/Source/Editor/AssetPlacementEdMode/Private/Tools/PlacementPlaceSingleTool.h b/Engine/Source/Editor/AssetPlacementEdMode/Private/Tools/PlacementPlaceSingleTool.h
--- a/Engine/Source/Editor/AssetPlacementEdMode/Private/Tools/PlacementPlaceSingleTool.h
+++ b/Engine/Source/Editor/AssetPlacementEdMode/Private/Tools/PlacementPlaceSingleTool.h
@@ -17,6 +17,30 @@ protected:
 	virtual UPlacementBrushToolBase* FactoryToolInstance(UObject* Outer) const override;
 };
 
+UENUM()
+enum class EPlacementScaleToCursorType : uint8
+{
+	/** Increases scale as the cursor moves outward from the placed asset. */
+	Positive,
+	/** Scale is unchanged by cursor movement for the placed asset. */
+	None,
+};
+
+UCLASS(config = EditorPerProjectUserSettings)
+class UPlacementModePlaceSingleToolSettings : public UInteractiveToolPropertySet
+{
+	GENERATED_BODY()
+
+public:
+	/** How the cursor movement should apply to scale once an asset is placed. Note that the maximum bound is controlled by scale settings of the mode. */
+	UPROPERTY(config, EditAnywhere, Category = "Single Place Settings")
+	EPlacementScaleToCursorType ScalingType = EPlacementScaleToCursorType::Positive;
+
+	/** If the tool should automatically select the last placed asset. */
+	UPROPERTY(config, EditAnywhere, Category = "Single Place Settings")
+	bool bSelectAfterPlacing = true;
+};
+
 UCLASS(Transient)
 class UPlacementModePlaceSingleTool : public UPlacementClickDragToolBase
 {
@@ -32,6 +56,8 @@ public:
 	virtual void Setup() override;
 	virtual void Shutdown(EToolShutdownType ShutdownType) override;
 	virtual void OnClickPress(const FInputDeviceRay& Ray) override;
+	virtual void OnClickDrag(const FInputDeviceRay& DragPos) override;
+	virtual void OnClickRelease(const FInputDeviceRay& ReleasePos) override;
 
 	virtual void OnBeginHover(const FInputDeviceRay& DevicePos) override;
 	virtual bool OnUpdateHover(const FInputDeviceRay& DevicePos) override;
@@ -39,20 +65,26 @@ public:
 	virtual FInputRayHit CanBeginClickDragSequence(const FInputDeviceRay& PressPos) override;
 	virtual FInputRayHit BeginHoverSequenceHitTest(const FInputDeviceRay& DevicePos) override;
 
+	virtual void Render(IToolsContextRenderAPI* RenderAPI) override;
+
 protected:
-	void GeneratePreviewPlacementData(const FInputDeviceRay& DevicePos);
+	void GeneratePlacementData(const FInputDeviceRay& DevicePos);
 	void CreatePreviewElements(const FInputDeviceRay& DevicePos);
 	void UpdatePreviewElements(const FInputDeviceRay& DevicePos);
 	void DestroyPreviewElements();
 	void EnterTweakState(TArrayView<const FTypedElementHandle> InElementHandles);
 	void ExitTweakState(bool bClearSelectionSet);
+	FAssetPlacementInfo UpdatePlacementInfo();
 	void UpdateElementTransforms(TArrayView<const FTypedElementHandle> InElements, const FTransform& InTransform, bool bLocalTransform = false);
 	void NotifyMovementStarted(TArrayView<const FTypedElementHandle> InElements);
 	void NotifyMovementEnded(TArrayView<const FTypedElementHandle> InElements);
 	void SetupRightClickMouseBehavior();
 
-	FQuat LastGeneratedRotation;
-	TUniquePtr<FAssetPlacementInfo> PreviewPlacementInfo;
+	UPROPERTY(Transient)
+	TObjectPtr<UPlacementModePlaceSingleToolSettings> SinglePlaceSettings;
+
+	TUniquePtr<FAssetPlacementInfo> PlacementInfo;
 	TArray<FTypedElementHandle> PreviewElements;
+	TArray<FTypedElementHandle> PlacedElements;
 	bool bIsTweaking;
 };
