Mirror content browser selection as palette in placement mode.

diff --git a/Engine/Source/Editor/AssetPlacementEdMode/AssetPlacementEdMode.Build.cs b/Engine/Source/Editor/AssetPlacementEdMode/AssetPlacementEdMode.Build.cs
--- a/Engine/Source/Editor/AssetPlacementEdMode/AssetPlacementEdMode.Build.cs
+++ b/Engine/Source/Editor/AssetPlacementEdMode/AssetPlacementEdMode.Build.cs
@@ -8,6 +8,7 @@ public class AssetPlacementEdMode : ModuleRules
 	{
 		PrivateDependencyModuleNames.AddRange(
 			new string[] {
+				"ContentBrowser",
 				"Core",
 				"CoreUObject",
                 "InputCore",
diff --git a/Engine/Source/Editor/AssetPlacementEdMode/Private/AssetPlacementSettings.h b/Engine/Source/Editor/AssetPlacementEdMode/Private/AssetPlacementSettings.h
--- a/Engine/Source/Editor/AssetPlacementEdMode/Private/AssetPlacementSettings.h
+++ b/Engine/Source/Editor/AssetPlacementEdMode/Private/AssetPlacementSettings.h
@@ -19,7 +19,7 @@ struct FPaletteItem
 	FAssetData AssetData;
 
 	UPROPERTY()
-	TScriptInterface<IAssetFactoryInterface> FactoryOverride;
+	TScriptInterface<IAssetFactoryInterface> AssetFactoryInterface;
 };
 
 UCLASS(config = EditorPerProjectUserSettings)
@@ -130,7 +130,10 @@ public:
 
 	// todo: palette save and load between sessions
 	//UPROPERTY(config)
-	TArray<FPaletteItem> PaletteItems;
+	TArray<TSharedPtr<FPaletteItem>> PaletteItems;
+
+	UPROPERTY(config)
+	bool bUseContentBrowserSelection = true;
 
 	virtual bool CanEditChange(const FProperty* InProperty) const override;
 };
diff --git a/Engine/Source/Editor/AssetPlacementEdMode/Private/Modes/PlacementModeSubsystem.cpp b/Engine/Source/Editor/AssetPlacementEdMode/Private/Modes/PlacementModeSubsystem.cpp
--- a/Engine/Source/Editor/AssetPlacementEdMode/Private/Modes/PlacementModeSubsystem.cpp
+++ b/Engine/Source/Editor/AssetPlacementEdMode/Private/Modes/PlacementModeSubsystem.cpp
@@ -8,6 +8,9 @@
 #include "Elements/Interfaces/TypedElementAssetDataInterface.h"
 #include "Factories/AssetFactoryInterface.h"
 
+#include "Subsystems/PlacementSubsystem.h"
+#include "Editor.h"
+
 void UPlacementModeSubsystem::Initialize(FSubsystemCollectionBase& Collection)
 {
 	ModeSettings = NewObject<UAssetPlacementSettings>(this);
@@ -35,9 +38,14 @@ bool UPlacementModeSubsystem::DoesCurrentPaletteSupportElement(const FTypedEleme
 	if (TTypedElement<UTypedElementAssetDataInterface> AssetDataInterface = UTypedElementRegistry::GetInstance()->GetElement<UTypedElementAssetDataInterface>(InElementToCheck))
 	{
 		TArray<FAssetData> ReferencedAssetDatas = AssetDataInterface.GetAllReferencedAssetDatas();
-		for (const FPaletteItem& Item : ModeSettings->PaletteItems)
+		for (const TSharedPtr<FPaletteItem>& Item : ModeSettings->PaletteItems)
 		{
-			if (ReferencedAssetDatas.Find(Item.AssetData) != INDEX_NONE)
+			if (!Item)
+			{
+				continue;
+			}
+
+			if (ReferencedAssetDatas.Find(Item->AssetData) != INDEX_NONE)
 			{
 				return true;
 			}
@@ -46,10 +54,15 @@ bool UPlacementModeSubsystem::DoesCurrentPaletteSupportElement(const FTypedEleme
 
 	// The current implementation of the asset data interface for actors requires that individual actors report on assets contained within their components.
 	// Not all actors do this reliably, so additionally check the supplied factory for a match. 
-	for (const FPaletteItem& Item : ModeSettings->PaletteItems)
+	for (const TSharedPtr<FPaletteItem>& Item : ModeSettings->PaletteItems)
 	{
-		FAssetData FoundAssetDataFromFactory = Item.FactoryOverride->GetAssetDataFromElementHandle(InElementToCheck);
-		if (FoundAssetDataFromFactory == Item.AssetData)
+		if (!Item)
+		{
+			continue;
+		}
+
+		FAssetData FoundAssetDataFromFactory = Item->AssetFactoryInterface->GetAssetDataFromElementHandle(InElementToCheck);
+		if (FoundAssetDataFromFactory == Item->AssetData)
 		{
 			return true;
 		}
@@ -58,18 +71,42 @@ bool UPlacementModeSubsystem::DoesCurrentPaletteSupportElement(const FTypedEleme
 	return false;
 }
 
-bool UPlacementModeSubsystem::AddPaletteItem(const FPaletteItem& InPaletteItem)
+TSharedPtr<FPaletteItem> UPlacementModeSubsystem::AddPaletteItem(const FAssetData& InAssetData)
 {
-	if (ModeSettings && !ModeSettings->PaletteItems.FindByPredicate([InPaletteItem](const FPaletteItem& ItemIter) { return ItemIter.AssetData.ObjectPath == InPaletteItem.AssetData.ObjectPath; }))
+	TSharedPtr<FPaletteItem> NewPaletteItem;
+	if (!InAssetData.IsValid())
 	{
-		ModeSettings->PaletteItems.Add(InPaletteItem);
-		return true;
+		return NewPaletteItem;
 	}
 
-	return false;
+	if (InAssetData.GetClass()->HasAnyClassFlags(CLASS_Abstract | CLASS_Deprecated | CLASS_NewerVersionExists | CLASS_NotPlaceable))
+	{
+		return NewPaletteItem;
+	}
+
+	if (ModeSettings && !ModeSettings->PaletteItems.FindByPredicate([InAssetData](const TSharedPtr<FPaletteItem>& ItemIter) { return ItemIter ? (ItemIter->AssetData.ObjectPath == InAssetData.ObjectPath) : false; }))
+	{
+		if (UPlacementSubsystem* PlacementSubystem = GEditor->GetEditorSubsystem<UPlacementSubsystem>())
+		{
+			if (TScriptInterface<IAssetFactoryInterface> AssetFactory = PlacementSubystem->FindAssetFactoryFromAssetData(InAssetData))
+			{
+				NewPaletteItem = MakeShared<FPaletteItem>();
+				NewPaletteItem->AssetData = InAssetData;
+				NewPaletteItem->AssetFactoryInterface = AssetFactory;
+				ModeSettings->PaletteItems.Add(NewPaletteItem);
+			}
+		}
+	}
+
+	return NewPaletteItem;
 }
 
 void UPlacementModeSubsystem::ClearPalette()
 {
 	ModeSettings->PaletteItems.Empty();
 }
+
+void UPlacementModeSubsystem::SetUseContentBrowserAsPalette(bool bInUseContentBrowser)
+{
+	ModeSettings->bUseContentBrowserSelection = bInUseContentBrowser;
+}
diff --git a/Engine/Source/Editor/AssetPlacementEdMode/Private/Modes/PlacementModeSubsystem.h b/Engine/Source/Editor/AssetPlacementEdMode/Private/Modes/PlacementModeSubsystem.h
--- a/Engine/Source/Editor/AssetPlacementEdMode/Private/Modes/PlacementModeSubsystem.h
+++ b/Engine/Source/Editor/AssetPlacementEdMode/Private/Modes/PlacementModeSubsystem.h
@@ -35,15 +35,20 @@ public:
 	/**
 	 * Adds the given palette item to the current palette.
 	 * 
-	 * @returns true if the item was added.
+	 * @returns a pointer to the added item, if the add was successful.
 	 */
-	bool AddPaletteItem(const FPaletteItem& InPaletteItem);
+	TSharedPtr<FPaletteItem> AddPaletteItem(const FAssetData& InAssetData);
 
 	/**
 	 * Clears all items from the current palette.
 	 */
 	void ClearPalette();
 
+	/**
+	 * Updates the settings object to use the content browser's active selection as the palette.
+	 */
+	void SetUseContentBrowserAsPalette(bool bInUseContentBrowser);
+
 protected:
 	UPROPERTY()
 	TObjectPtr<UAssetPlacementSettings> ModeSettings;
diff --git a/Engine/Source/Editor/AssetPlacementEdMode/Private/SAssetPlacementPalette.cpp b/Engine/Source/Editor/AssetPlacementEdMode/Private/SAssetPlacementPalette.cpp
--- a/Engine/Source/Editor/AssetPlacementEdMode/Private/SAssetPlacementPalette.cpp
+++ b/Engine/Source/Editor/AssetPlacementEdMode/Private/SAssetPlacementPalette.cpp
@@ -1,21 +1,15 @@
 // Copyright Epic Games, Inc. All Rights Reserved.
 
 #include "SAssetPlacementPalette.h"
-#include "Widgets/Layout/SSplitter.h"
-#include "Widgets/Input/SButton.h"
-#include "Widgets/Input/SCheckBox.h"
 #include "AssetPlacementPaletteItem.h"
-#include "Misc/MessageDialog.h"
-#include "Misc/FeedbackContext.h"
-#include "Modules/ModuleManager.h"
 #include "Framework/Application/SlateApplication.h"
 #include "Framework/Commands/UIAction.h"
 #include "Framework/Commands/UICommandList.h"
-#include "Widgets/Images/SImage.h"
-#include "Widgets/Text/SRichTextBlock.h"
 #include "Framework/MultiBox/MultiBoxBuilder.h"
+#include "Widgets/Images/SImage.h"
 #include "Widgets/Layout/SScrollBorder.h"
 #include "Widgets/Input/SSlider.h"
+#include "Widgets/Input/SCheckBox.h"
 #include "EditorStyleSet.h"
 #include "AssetThumbnail.h"
 #include "PropertyEditorModule.h"
@@ -24,9 +18,8 @@
 #include "ContentBrowserModule.h"
 #include "PropertyCustomizationHelpers.h"
 #include "IDetailsView.h"
-#include "AssetSelection.h"
-#include "ScopedTransaction.h"
 #include "AssetData.h"
+#include "AssetSelection.h"
 #include "Editor.h"
 #include "Widgets/Layout/SScrollBox.h"
 
@@ -40,6 +33,10 @@
 #include "Engine/StreamableManager.h"
 #include "Modes/PlacementModeSubsystem.h"
 
+#include "IContentBrowserSingleton.h"
+#include "ContentBrowserModule.h"
+#include "Modules/ModuleManager.h"
+
 #define LOCTEXT_NAMESPACE "AssetPlacementMode"
 
 ////////////////////////////////////////////////
@@ -123,6 +120,22 @@ void SAssetPlacementPalette::Construct(const FArguments& InArgs)
 
 	const FText BlankText = FText::GetEmpty();
 
+	TWeakObjectPtr<const UAssetPlacementSettings> PlacementSettings = GEditor->GetEditorSubsystem<UPlacementModeSubsystem>()->GetModeSettingsObject();
+	if (PlacementSettings.IsValid())
+	{
+		for (const TSharedPtr<FPaletteItem>& PaletteItem : PlacementSettings->PaletteItems)
+		{
+			if (PaletteItem)
+			{
+				PaletteItems.Add(MakeShared<FAssetPlacementPaletteItemModel>(PaletteItem, SharedThis(this), ThumbnailPool));
+			}
+		}
+
+		// Make sure the content browser setting is not what we're going to change to, so that the setup will happen.
+		bIsMirroringContentBrowser = !PlacementSettings->bUseContentBrowserSelection;
+		SetupContentBrowserMirroring(PlacementSettings->bUseContentBrowserSelection);
+	}
+
 	ChildSlot
 	[
 		SNew(SVerticalBox)
@@ -136,8 +149,27 @@ void SAssetPlacementPalette::Construct(const FArguments& InArgs)
 			.BorderBackgroundColor(FLinearColor(.6f, .6f, .6f, 1.0f))
 			[
 				SNew(SVerticalBox)
+				+ SVerticalBox::Slot()
+				.AutoHeight()
+				.HAlign(HAlign_Fill)
+				[
+					SNew(SCheckBox)
+					.Type(ESlateCheckBoxType::Type::ToggleButton)
+					.IsChecked(bIsMirroringContentBrowser ? ECheckBoxState::Checked : ECheckBoxState::Unchecked)
+					.Style(&FAssetPlacementEdModeStyle::Get().GetWidgetStyle<FCheckBoxStyle>("ToggleButtonCheckBox"))
+					.OnCheckStateChanged(this, &SAssetPlacementPalette::OnContentBrowserMirrorButtonClicked)
+					[
+						SNew(STextBlock)
+						.Text(LOCTEXT("Placement_ToggleContentBrowserMirroring", "Mirror Content Browser Selection"))
+						.Justification(ETextJustify::Type::Center)
+						.TextStyle(&FAssetPlacementEdModeStyle::Get().GetWidgetStyle<FTextBlockStyle>("ButtonText"))
+						.ToolTipText(LOCTEXT("Placement_ToggleContentBrowserMirroring_ToolTip", "Toggles palette to mirror the active content browser selection."))
+					]
+				]
+
 				+ SVerticalBox::Slot()
 				.VAlign(VAlign_Center)
+				.AutoHeight()
 				[
 					SNew(SHorizontalBox)
 
@@ -223,15 +255,6 @@ void SAssetPlacementPalette::Construct(const FArguments& InArgs)
 		]
 	];
 
-	TWeakObjectPtr<const UAssetPlacementSettings> PlacementSettings = GEditor->GetEditorSubsystem<UPlacementModeSubsystem>()->GetModeSettingsObject();
-	if (PlacementSettings.IsValid())
-	{
-		for (const FPaletteItem& PaletteItem : PlacementSettings->PaletteItems)
-		{
-			PaletteItems.Add(MakeShared<FAssetPlacementPaletteItemModel>(MakeShared<FPaletteItem>(PaletteItem), SharedThis(this), ThumbnailPool));
-		}
-	}
-
 	UpdatePalette(true);
 }
 
@@ -284,33 +307,9 @@ void SAssetPlacementPalette::RefreshActivePaletteViewWidget()
 
 void SAssetPlacementPalette::AddPlacementType(const FAssetData& AssetData)
 {
-	if (!AssetData.IsValid())
-	{
-		return;
-	}
-
-	if (AssetData.GetClass()->HasAnyClassFlags(CLASS_Abstract | CLASS_Deprecated | CLASS_NewerVersionExists | CLASS_NotPlaceable))
-	{
-		return;
-	}
-
-	TScriptInterface<IAssetFactoryInterface> FactoryInterface;
-	if (UPlacementSubsystem* PlacementSubystem = GEditor->GetEditorSubsystem<UPlacementSubsystem>())
-	{
-		FactoryInterface = PlacementSubystem->FindAssetFactoryFromAssetData(AssetData);
-	}
-
-	if (!FactoryInterface)
-	{
-		return;
-	}
-
-	FAssetPlacementUIInfoPtr PlacementInfo = MakeShared<FPaletteItem>();
-	PlacementInfo->AssetData = AssetData;
-	PlacementInfo->FactoryOverride = FactoryInterface;
-
 	// Try to add the item to the mode's palette
-	if (!GEditor->GetEditorSubsystem<UPlacementModeSubsystem>()->AddPaletteItem(*PlacementInfo))
+	FAssetPlacementUIInfoPtr PlacementInfo = GEditor->GetEditorSubsystem<UPlacementModeSubsystem>()->AddPaletteItem(AssetData);
+	if (!PlacementInfo.IsValid())
 	{
 		return;
 	}
@@ -338,6 +337,15 @@ void SAssetPlacementPalette::OnClearPalette()
 	UpdatePalette(true);
 }
 
+void SAssetPlacementPalette::SetPaletteToAssetDataList(TArrayView<const FAssetData> InAssetDatas)
+{
+	ClearPalette();
+	for (const FAssetData& SelectedAsset : InAssetDatas)
+	{
+		AddPlacementType(SelectedAsset);
+	}
+}
+
 TSharedRef<SWidgetSwitcher> SAssetPlacementPalette::CreatePaletteViews()
 {
 	const FText BlankText = FText::GetEmpty();
@@ -427,6 +435,48 @@ bool SAssetPlacementPalette::ShouldFilterAsset(const FAssetData& InAssetData)
 	return true;
 }
 
+void SAssetPlacementPalette::OnContentBrowserMirrorButtonClicked(ECheckBoxState InState)
+{
+	SetupContentBrowserMirroring((InState == ECheckBoxState::Checked));
+
+	if (UPlacementModeSubsystem* PlacementModeSubsystem = GEditor->GetEditorSubsystem<UPlacementModeSubsystem>())
+	{
+		PlacementModeSubsystem->SetUseContentBrowserAsPalette(bIsMirroringContentBrowser);
+	}
+}
+
+void SAssetPlacementPalette::OnContentBrowserSelectionChanged(const TArray<FAssetData>& NewSelectedAssets, bool bIsPrimaryBrowser)
+{
+	if (bIsPrimaryBrowser)
+	{
+		SetPaletteToAssetDataList(NewSelectedAssets);
+		UpdatePalette(true);
+	}
+}
+
+void SAssetPlacementPalette::SetupContentBrowserMirroring(bool bInMirrorContentBrowser)
+{
+	bool bWasMirroringContentBrowser = bIsMirroringContentBrowser;
+	if (bWasMirroringContentBrowser != bInMirrorContentBrowser)
+	{
+		if (FContentBrowserModule* ContentBrowserModule = FModuleManager::GetModulePtr<FContentBrowserModule>("ContentBrowser"))
+		{
+			if (bInMirrorContentBrowser)
+			{
+				TArray<FAssetData> SelectedAssetDatas;
+				ContentBrowserModule->Get().GetSelectedAssets(SelectedAssetDatas);
+				OnContentBrowserSelectionChanged(SelectedAssetDatas, true);
+				ContentBrowserModule->GetOnAssetSelectionChanged().AddSP(SharedThis(this), &SAssetPlacementPalette::OnContentBrowserSelectionChanged);
+			}
+			else
+			{
+				ContentBrowserModule->GetOnAssetSelectionChanged().RemoveAll(this);
+			}
+		}
+	}
+	bIsMirroringContentBrowser = bInMirrorContentBrowser;
+}
+
 void SAssetPlacementPalette::SetViewMode(EViewMode NewViewMode)
 {
 	if (ActiveViewMode != NewViewMode)
@@ -549,11 +599,21 @@ TSharedPtr<SListView<FPlacementPaletteItemModelPtr>> SAssetPlacementPalette::Get
 
 EVisibility SAssetPlacementPalette::GetDropPlacementHintVisibility() const
 {
+	if (bIsMirroringContentBrowser)
+	{
+		return EVisibility::Collapsed;
+	}
+
 	return FilteredItems.Num() ? EVisibility::Collapsed : EVisibility::Visible;
 }
 
 EVisibility SAssetPlacementPalette::GetPlacementDropTargetVisibility() const
 {
+	if (bIsMirroringContentBrowser)
+	{
+		return EVisibility::Hidden;
+	}
+
 	if (FSlateApplication::Get().IsDragDropping())
 	{
 		TArray<FAssetData> DraggedAssets = AssetUtil::ExtractAssetDataFromDrag(FSlateApplication::Get().GetDragDroppingContent());
@@ -576,6 +636,11 @@ EVisibility SAssetPlacementPalette::GetPlacementDropTargetVisibility() const
 
 FReply SAssetPlacementPalette::HandlePlacementDropped(const FGeometry& DropZoneGeometry, const FDragDropEvent& DragDropEvent)
 {
+	if (bIsMirroringContentBrowser)
+	{
+		return FReply::Unhandled();
+	}
+
 	TArray<FAssetData> DroppedAssetData = AssetUtil::ExtractAssetDataFromDrag(DragDropEvent);
 	if (DroppedAssetData.Num() > 0)
 	{
@@ -603,12 +668,17 @@ bool SAssetPlacementPalette::HasAnyItemInPalette() const
 TSharedPtr<SWidget> SAssetPlacementPalette::ConstructPlacementTypeContextMenu()
 {
 	FMenuBuilder MenuBuilder(true, nullptr);
-	MenuBuilder.AddMenuEntry(
-		LOCTEXT("Palette_Clear", "Clear Palette"),
-		LOCTEXT("Palette_ClearDesc", "Removes all items from the palette."),
-		FSlateIcon(),
-		FUIAction(FExecuteAction::CreateSP(this, &SAssetPlacementPalette::OnClearPalette), FCanExecuteAction::CreateSP(this, &SAssetPlacementPalette::HasAnyItemInPalette))
-	);
+	if (!bIsMirroringContentBrowser)
+	{
+		MenuBuilder.AddMenuEntry(
+			LOCTEXT("Palette_Clear", "Clear Palette"),
+			LOCTEXT("Palette_ClearDesc", "Removes all items from the palette."),
+			FSlateIcon(),
+			FUIAction(
+				FExecuteAction::CreateSP(this, &SAssetPlacementPalette::OnClearPalette),
+				FCanExecuteAction::CreateSP(this, &SAssetPlacementPalette::HasAnyItemInPalette))
+		);
+	}
 	return MenuBuilder.MakeWidget();
 }
 
diff --git a/Engine/Source/Editor/AssetPlacementEdMode/Private/SAssetPlacementPalette.h b/Engine/Source/Editor/AssetPlacementEdMode/Private/SAssetPlacementPalette.h
--- a/Engine/Source/Editor/AssetPlacementEdMode/Private/SAssetPlacementPalette.h
+++ b/Engine/Source/Editor/AssetPlacementEdMode/Private/SAssetPlacementPalette.h
@@ -76,6 +76,7 @@ private:
 	 */
 	void ClearPalette();
 	void OnClearPalette();
+	void SetPaletteToAssetDataList(TArrayView<const FAssetData> InAssetDatas);
 
 	/** Refreshes the active palette view widget */
 	void RefreshActivePaletteViewWidget();
@@ -91,8 +92,9 @@ private:
 
 	bool ShouldFilterAsset(const FAssetData& InAssetData);
 
-	/** Gets the visibility of the Add Placement Type text in the header row button */
-	EVisibility GetAddPlacementTypeButtonTextVisibility() const;
+	void OnContentBrowserMirrorButtonClicked(ECheckBoxState InState);
+	void OnContentBrowserSelectionChanged(const TArray<FAssetData>& NewSelectedAssets, bool bIsPrimaryBrowser);
+	void SetupContentBrowserMirroring(bool bInMirrorContentBrowser);
 
 	/** Sets the view mode of the palette */
 	void SetViewMode(EViewMode NewViewMode);
@@ -212,6 +214,7 @@ private:
 	bool bShowFullTooltips = true;
 	bool bIsRebuildTimerRegistered = true;
 	bool bIsRefreshTimerRegistered = true;
+	bool bIsMirroringContentBrowser = true;
 	EViewMode ActiveViewMode;
 	EColumnSortMode::Type ActiveSortOrder = EColumnSortMode::Type::Ascending;
 
diff --git a/Engine/Source/Editor/AssetPlacementEdMode/Private/Tools/PlacementPlaceSingleTool.cpp b/Engine/Source/Editor/AssetPlacementEdMode/Private/Tools/PlacementPlaceSingleTool.cpp
--- a/Engine/Source/Editor/AssetPlacementEdMode/Private/Tools/PlacementPlaceSingleTool.cpp
+++ b/Engine/Source/Editor/AssetPlacementEdMode/Private/Tools/PlacementPlaceSingleTool.cpp
@@ -235,12 +235,17 @@ void UPlacementModePlaceSingleTool::GeneratePlacementData(const FInputDeviceRay&
 	if (PlacementSettings && PlacementSettings->PaletteItems.Num())
 	{
 		int32 ItemIndex = FMath::RandHelper(PlacementSettings->PaletteItems.Num());
-		const FPaletteItem& ItemToPlace = PlacementSettings->PaletteItems[ItemIndex];
+		const TSharedPtr<FPaletteItem>& ItemToPlace = PlacementSettings->PaletteItems[ItemIndex];
+		if (!ItemToPlace)
+		{
+			return;
+		}
+
 		FTransform TransformToUpdate(GenerateRandomRotation(PlacementSettings), LastBrushStamp.WorldPosition, GenerateRandomScale(PlacementSettings));
 
 		PlacementInfo = MakeUnique<FAssetPlacementInfo>();
-		PlacementInfo->AssetToPlace = ItemToPlace.AssetData;
-		PlacementInfo->FactoryOverride = ItemToPlace.FactoryOverride;
+		PlacementInfo->AssetToPlace = ItemToPlace->AssetData;
+		PlacementInfo->FactoryOverride = ItemToPlace->AssetFactoryInterface;
 		PlacementInfo->FinalizedTransform = FinalizeTransform(TransformToUpdate, LastBrushStamp.WorldNormal, PlacementSettings);
 		PlacementInfo->PreferredLevel = GEditor->GetEditorWorldContext().World()->GetCurrentLevel();
 	}
diff --git a/Engine/Source/Editor/AssetPlacementEdMode/Private/Tools/PlacementPlaceTool.cpp b/Engine/Source/Editor/AssetPlacementEdMode/Private/Tools/PlacementPlaceTool.cpp
--- a/Engine/Source/Editor/AssetPlacementEdMode/Private/Tools/PlacementPlaceTool.cpp
+++ b/Engine/Source/Editor/AssetPlacementEdMode/Private/Tools/PlacementPlaceTool.cpp
@@ -86,9 +86,15 @@ void UPlacementModePlacementTool::OnTick(float DeltaTime)
 			PotentialInstanceHash.InsertInstance(SpawnLocation, PotentialIdx);
 
 			int32 ItemIndex = FMath::RandHelper(PlacementSettings->PaletteItems.Num());
-			const FPaletteItem& ItemToPlace = PlacementSettings->PaletteItems[ItemIndex];
+			const TSharedPtr<FPaletteItem>& ItemToPlace = PlacementSettings->PaletteItems[ItemIndex];
+			if (!ItemToPlace)
+			{
+				return;
+			}
+
 			FAssetPlacementInfo NewInfo;
-			NewInfo.AssetToPlace = ItemToPlace.AssetData;
+			NewInfo.AssetToPlace = ItemToPlace->AssetData;
+			NewInfo.FactoryOverride = ItemToPlace->AssetFactoryInterface;
 			NewInfo.PreferredLevel = GEditor->GetEditorWorldContext().World()->GetCurrentLevel();
 			NewInfo.FinalizedTransform = GenerateTransformFromHitLocationAndNormal(SpawnLocation, SpawnNormal);
 
