Adds AssetPlacementMode Icon and Updates Style to use Macors

diff --git a/Engine/Source/Editor/AssetPlacementEdMode/Private/AssetPlacementEdMode.cpp b/Engine/Source/Editor/AssetPlacementEdMode/Private/AssetPlacementEdMode.cpp
--- a/Engine/Source/Editor/AssetPlacementEdMode/Private/AssetPlacementEdMode.cpp
+++ b/Engine/Source/Editor/AssetPlacementEdMode/Private/AssetPlacementEdMode.cpp
@@ -40,7 +40,7 @@ UAssetPlacementEdMode::UAssetPlacementEdMode()
 	TAttribute<bool> IsEnabledAttr = TAttribute<bool>::Create(TAttribute<bool>::FGetter::CreateLambda([]() -> bool { return GetDefault<ULevelEditorMiscSettings>()->bEnableAssetPlacementMode; }));
 	Info = FEditorModeInfo(UAssetPlacementEdMode::AssetPlacementEdModeID,
 		LOCTEXT("AssetPlacementEdModeName", "Placement"),
-		FSlateIcon(FAssetPlacementEdModeStyle::GetStyleSetName(), "PlacementBrowser.ShowAllContent"),
+		FSlateIcon(FAssetPlacementEdModeStyle::Get().GetStyleSetName(), "LevelEditor.AssetPlacementEdMode"),
 		MoveTemp(IsEnabledAttr));
 }
 
diff --git a/Engine/Source/Editor/AssetPlacementEdMode/Private/AssetPlacementEdModeCommands.cpp b/Engine/Source/Editor/AssetPlacementEdMode/Private/AssetPlacementEdModeCommands.cpp
--- a/Engine/Source/Editor/AssetPlacementEdMode/Private/AssetPlacementEdModeCommands.cpp
+++ b/Engine/Source/Editor/AssetPlacementEdMode/Private/AssetPlacementEdModeCommands.cpp
@@ -9,7 +9,7 @@ FAssetPlacementEdModeCommands::FAssetPlacementEdModeCommands()
 	: TCommands<FAssetPlacementEdModeCommands>("AssetPlacementEdMode",
 		LOCTEXT("AssetPlacementEdModeCommands", "AssetPlacement Editor Mode Commands"),
 		NAME_None,
-		FAssetPlacementEdModeStyle::GetStyleSetName())
+		FAssetPlacementEdModeStyle::Get().GetStyleSetName())
 {
 }
 
diff --git a/Engine/Source/Editor/AssetPlacementEdMode/Private/AssetPlacementEdModeModule.cpp b/Engine/Source/Editor/AssetPlacementEdMode/Private/AssetPlacementEdModeModule.cpp
--- a/Engine/Source/Editor/AssetPlacementEdMode/Private/AssetPlacementEdModeModule.cpp
+++ b/Engine/Source/Editor/AssetPlacementEdMode/Private/AssetPlacementEdModeModule.cpp
@@ -11,7 +11,7 @@ IMPLEMENT_MODULE( FAssetPlacementEdMode, AssetPlacementEdMode );
 
 void FAssetPlacementEdMode::StartupModule()
 {
-	FAssetPlacementEdModeStyle::Initialize();
+	FAssetPlacementEdModeStyle::Get();
 	FAssetPlacementEdModeCommands::Register();
 }
 
diff --git a/Engine/Source/Editor/AssetPlacementEdMode/Private/AssetPlacementEdModeStyle.cpp b/Engine/Source/Editor/AssetPlacementEdMode/Private/AssetPlacementEdModeStyle.cpp
--- a/Engine/Source/Editor/AssetPlacementEdMode/Private/AssetPlacementEdModeStyle.cpp
+++ b/Engine/Source/Editor/AssetPlacementEdMode/Private/AssetPlacementEdModeStyle.cpp
@@ -10,94 +10,80 @@
 #include "Styling/StarshipCoreStyle.h"
 #include "Styling/SlateStyleRegistry.h"
 
-TSharedPtr<FAssetPlacementEdModeStyle> FAssetPlacementEdModeStyle::AssetPlacementEdModeStyle;
+FName FAssetPlacementEdModeStyle::StyleName("AssetPlacementEdModeStyle");
+TUniquePtr<FAssetPlacementEdModeStyle> FAssetPlacementEdModeStyle::Inst(nullptr);
 
-void FAssetPlacementEdModeStyle::Initialize()
+const FName& FAssetPlacementEdModeStyle::GetStyleSetName() const
 {
-	if (AssetPlacementEdModeStyle.IsValid())
-	{
-		return;
-	}
-
-	AssetPlacementEdModeStyle = MakeShareable(new FAssetPlacementEdModeStyle);
-	AssetPlacementEdModeStyle->SetupCustomStyle();
-}
-
-void FAssetPlacementEdModeStyle::Shutdown()
-{
-	AssetPlacementEdModeStyle.Reset();
+	return StyleName;
 }
 
-FName FAssetPlacementEdModeStyle::GetStyleSetName()
+const FAssetPlacementEdModeStyle& FAssetPlacementEdModeStyle::Get()
 {
-	return "AssetPlacementEdModeStyle";
+	if (!Inst.IsValid())
+	{
+		Inst = TUniquePtr<FAssetPlacementEdModeStyle>(new FAssetPlacementEdModeStyle);
+	}
+	return *(Inst.Get());
 }
 
-const ISlateStyle& FAssetPlacementEdModeStyle::Get()
+void FAssetPlacementEdModeStyle::Shutdown()
 {
-	check(AssetPlacementEdModeStyle.IsValid());
-	return *AssetPlacementEdModeStyle->StyleSet;
+	Inst.Reset();
 }
 
 FAssetPlacementEdModeStyle::FAssetPlacementEdModeStyle()
+	: FSlateStyleSet(StyleName)
 {
-	StyleSet = MakeShareable(new FSlateStyleSet(GetStyleSetName()));
 
-	StyleSet->SetParentStyleName("EditorStyle");
-	StyleSet->SetContentRoot(FPaths::EngineContentDir() / TEXT("Editor/Slate"));
-	StyleSet->SetCoreContentRoot(FPaths::EngineContentDir() / TEXT("Slate"));
-}
+	const FVector2D Icon20x20(20.0f, 20.0f);
 
-void FAssetPlacementEdModeStyle::SetupCustomStyle()
-{
-	auto MakeImageBrushFn = [this](const FString& RelativePath)
-	{
-		static const FVector2D Icon20x20(20.0f, 20.0f);
-		constexpr const TCHAR ImageExtension[] = TEXT(".png");
-		return new FSlateImageBrush(StyleSet->RootToContentDir(RelativePath, ImageExtension), Icon20x20);
-	};
+	SetParentStyleName("EditorStyle");
+	SetContentRoot(FPaths::EngineContentDir() / TEXT("Editor/Slate"));
+	SetCoreContentRoot(FPaths::EngineContentDir() / TEXT("Slate"));
+
+	Set("LevelEditor.AssetPlacementEdMode", new IMAGE_BRUSH_SVG("Starship/MainToolbar/AssetPlacementMode", FVector2D(20.0f, 20.0f)));
 
-	StyleSet->Set("AssetPlacementEdMode.Select", MakeImageBrushFn("Icons/GeneralTools/Select_40x"));
-	StyleSet->Set("AssetPlacementEdMode.Select.Small", MakeImageBrushFn("Icons/GeneralTools/Select_40x"));
-	StyleSet->Set("AssetPlacementEdMode.SelectAll", MakeImageBrushFn("Icons/GeneralTools/SelectAll_40x"));
-	StyleSet->Set("AssetPlacementEdMode.SelectAll.Small", MakeImageBrushFn("Icons/GeneralTools/SelectAll_40x"));
-	StyleSet->Set("AssetPlacementEdMode.Deselect", MakeImageBrushFn("Icons/GeneralTools/Deselect_40x"));
-	StyleSet->Set("AssetPlacementEdMode.Deselect.Small", MakeImageBrushFn("Icons/GeneralTools/Deselect_40x"));
-	StyleSet->Set("AssetPlacementEdMode.SelectInvalid", MakeImageBrushFn("Icons/GeneralTools/SelectInvalid_40x"));
-	StyleSet->Set("AssetPlacementEdMode.SelectInvalid.Small", MakeImageBrushFn("Icons/GeneralTools/SelectInvalid_40x"));
-	StyleSet->Set("AssetPlacementEdMode.LassoSelect", MakeImageBrushFn("Icons/GeneralTools/Lasso_40x"));
-	StyleSet->Set("AssetPlacementEdMode.LassoSelect.Small", MakeImageBrushFn("Icons/GeneralTools/Lasso_40x"));
-	StyleSet->Set("AssetPlacementEdMode.PlaceSingle", MakeImageBrushFn("Icons/GeneralTools/Foliage_40x"));
-	StyleSet->Set("AssetPlacementEdMode.PlaceSingle.Small", MakeImageBrushFn("Icons/GeneralTools/Foliage_40x"));
-	StyleSet->Set("AssetPlacementEdMode.Place", MakeImageBrushFn("Icons/GeneralTools/Paint_40x"));
-	StyleSet->Set("AssetPlacementEdMode.Place.Small", MakeImageBrushFn("Icons/GeneralTools/Paint_40x"));
-	StyleSet->Set("AssetPlacementEdMode.ReapplySettings", MakeImageBrushFn("Icons/GeneralTools/Repaint_40x"));
-	StyleSet->Set("AssetPlacementEdMode.ReapplySettings.Small", MakeImageBrushFn("Icons/GeneralTools/Repaint_40x"));
-	StyleSet->Set("AssetPlacementEdMode.Fill", MakeImageBrushFn("Icons/GeneralTools/PaintBucket_40x"));
-	StyleSet->Set("AssetPlacementEdMode.Fill.Small", MakeImageBrushFn("Icons/GeneralTools/PaintBucket_40x"));
-	StyleSet->Set("AssetPlacementEdMode.Delete", MakeImageBrushFn("Icons/GeneralTools/Delete_40x"));
-	StyleSet->Set("AssetPlacementEdMode.Delete.Small", MakeImageBrushFn("Icons/GeneralTools/Delete_40x"));
-	StyleSet->Set("AssetPlacementEdMode.Erase", MakeImageBrushFn("Icons/GeneralTools/Erase_40x"));
-	StyleSet->Set("AssetPlacementEdMode.Erase.Small", MakeImageBrushFn("Icons/GeneralTools/Erase_40x"));
-	StyleSet->Set("AssetPlacementEdMode.Filter", MakeImageBrushFn("Icons/GeneralTools/Filter_40x"));
-	StyleSet->Set("AssetPlacementEdMode.Filter.Small", MakeImageBrushFn("Icons/GeneralTools/Filter_40x"));
-	StyleSet->Set("AssetPlacementEdMode.Settings", MakeImageBrushFn("Icons/GeneralTools/Settings_40x"));
-	StyleSet->Set("AssetPlacementEdMode.Settings.Small", MakeImageBrushFn("Icons/GeneralTools/Settings_40x"));
-	StyleSet->Set("AssetPlacementEdMode.MoveToActivePartition", MakeImageBrushFn("Icons/GeneralTools/MoveToLevel_40x"));
-	StyleSet->Set("AssetPlacementEdMode.MoveToActivePartition.Small", MakeImageBrushFn("Icons/GeneralTools/MoveToLevel_40x"));
+	Set("AssetPlacementEdMode.Select",                      new IMAGE_BRUSH("Icons/GeneralTools/Select_40x", Icon20x20));
+	Set("AssetPlacementEdMode.Select.Small",                new IMAGE_BRUSH("Icons/GeneralTools/Select_40x", Icon20x20));
+	Set("AssetPlacementEdMode.SelectAll",                   new IMAGE_BRUSH("Icons/GeneralTools/SelectAll_40x", Icon20x20));
+	Set("AssetPlacementEdMode.SelectAll.Small",             new IMAGE_BRUSH("Icons/GeneralTools/SelectAll_40x", Icon20x20));
+	Set("AssetPlacementEdMode.Deselect",                    new IMAGE_BRUSH("Icons/GeneralTools/Deselect_40x", Icon20x20));
+	Set("AssetPlacementEdMode.Deselect.Small",              new IMAGE_BRUSH("Icons/GeneralTools/Deselect_40x", Icon20x20));
+	Set("AssetPlacementEdMode.SelectInvalid",               new IMAGE_BRUSH("Icons/GeneralTools/SelectInvalid_40x", Icon20x20));
+	Set("AssetPlacementEdMode.SelectInvalid.Small",         new IMAGE_BRUSH("Icons/GeneralTools/SelectInvalid_40x", Icon20x20));
+	Set("AssetPlacementEdMode.LassoSelect",                 new IMAGE_BRUSH("Icons/GeneralTools/Lasso_40x", Icon20x20));
+	Set("AssetPlacementEdMode.LassoSelect.Small",           new IMAGE_BRUSH("Icons/GeneralTools/Lasso_40x", Icon20x20));
+	Set("AssetPlacementEdMode.PlaceSingle",                 new IMAGE_BRUSH("Icons/GeneralTools/Foliage_40x", Icon20x20));
+	Set("AssetPlacementEdMode.PlaceSingle.Small",           new IMAGE_BRUSH("Icons/GeneralTools/Foliage_40x", Icon20x20));
+	Set("AssetPlacementEdMode.Place",                       new IMAGE_BRUSH("Icons/GeneralTools/Paint_40x", Icon20x20));
+	Set("AssetPlacementEdMode.Place.Small",                 new IMAGE_BRUSH("Icons/GeneralTools/Paint_40x", Icon20x20));
+	Set("AssetPlacementEdMode.ReapplySettings",             new IMAGE_BRUSH("Icons/GeneralTools/Repaint_40x", Icon20x20));
+	Set("AssetPlacementEdMode.ReapplySettings.Small",       new IMAGE_BRUSH("Icons/GeneralTools/Repaint_40x", Icon20x20));
+	Set("AssetPlacementEdMode.Fill",                        new IMAGE_BRUSH("Icons/GeneralTools/PaintBucket_40x", Icon20x20));
+	Set("AssetPlacementEdMode.Fill.Small",                  new IMAGE_BRUSH("Icons/GeneralTools/PaintBucket_40x", Icon20x20));
+	Set("AssetPlacementEdMode.Delete",                      new IMAGE_BRUSH("Icons/GeneralTools/Delete_40x", Icon20x20));
+	Set("AssetPlacementEdMode.Delete.Small",                new IMAGE_BRUSH("Icons/GeneralTools/Delete_40x", Icon20x20));
+	Set("AssetPlacementEdMode.Erase",                       new IMAGE_BRUSH("Icons/GeneralTools/Erase_40x", Icon20x20));
+	Set("AssetPlacementEdMode.Erase.Small",                 new IMAGE_BRUSH("Icons/GeneralTools/Erase_40x", Icon20x20));
+	Set("AssetPlacementEdMode.Filter",                      new IMAGE_BRUSH("Icons/GeneralTools/Filter_40x", Icon20x20));
+	Set("AssetPlacementEdMode.Filter.Small",                new IMAGE_BRUSH("Icons/GeneralTools/Filter_40x", Icon20x20));
+	Set("AssetPlacementEdMode.Settings",                    new IMAGE_BRUSH("Icons/GeneralTools/Settings_40x", Icon20x20));
+	Set("AssetPlacementEdMode.Settings.Small",              new IMAGE_BRUSH("Icons/GeneralTools/Settings_40x", Icon20x20));
+	Set("AssetPlacementEdMode.MoveToActivePartition",       new IMAGE_BRUSH("Icons/GeneralTools/MoveToLevel_40x", Icon20x20));
+	Set("AssetPlacementEdMode.MoveToActivePartition.Small", new IMAGE_BRUSH("Icons/GeneralTools/MoveToLevel_40x", Icon20x20));
 
-	StyleSet->Set("AssetPlacementEdMode.AddAssetType.Text", FTextBlockStyle(StyleSet->GetWidgetStyle<FTextBlockStyle>("NormalText"))
+	Set("AssetPlacementEdMode.AddAssetType.Text", FTextBlockStyle(GetWidgetStyle<FTextBlockStyle>("NormalText"))
 		.SetFont(DEFAULT_FONT("Bold", 10))
 		.SetColorAndOpacity(FLinearColor(1.0f, 1.0f, 1.0f))
 		.SetHighlightColor(FLinearColor(1.0f, 1.0f, 1.0f))
 		.SetShadowOffset(FVector2D(1, 1))
 		.SetShadowColorAndOpacity(FLinearColor(0, 0, 0, 0.9f)));
 
-	FSlateStyleRegistry::RegisterSlateStyle(*StyleSet.Get());
+	FSlateStyleRegistry::RegisterSlateStyle(*this);
 }
 
 FAssetPlacementEdModeStyle::~FAssetPlacementEdModeStyle()
 {
-	FSlateStyleRegistry::UnRegisterSlateStyle(*StyleSet.Get());
-	StyleSet.Reset();
+	FSlateStyleRegistry::UnRegisterSlateStyle(*this);
 }
diff --git a/Engine/Source/Editor/AssetPlacementEdMode/Private/AssetPlacementEdModeStyle.h b/Engine/Source/Editor/AssetPlacementEdMode/Private/AssetPlacementEdModeStyle.h
--- a/Engine/Source/Editor/AssetPlacementEdMode/Private/AssetPlacementEdModeStyle.h
+++ b/Engine/Source/Editor/AssetPlacementEdMode/Private/AssetPlacementEdModeStyle.h
@@ -3,27 +3,23 @@
 #pragma once
 
 #include "CoreMinimal.h"
-
-class FSlateStyleSet;
-class ISlateStyle;
+#include "Styling/SlateStyle.h"
+#include "Styling/SlateStyleRegistry.h"
 
 class FAssetPlacementEdModeStyle
+	: public FSlateStyleSet
 {
 public:
-	~FAssetPlacementEdModeStyle();
-
-	static void Initialize();
+	virtual const FName& GetStyleSetName() const override;
 
+	static const FAssetPlacementEdModeStyle& Get();
 	static void Shutdown();
-
-	static FName GetStyleSetName();
-	static const ISlateStyle& Get();
+	
+	~FAssetPlacementEdModeStyle();
 
 private:
 	FAssetPlacementEdModeStyle();
 
-	void SetupCustomStyle();
-
-	TSharedPtr<FSlateStyleSet> StyleSet;
-	static TSharedPtr<FAssetPlacementEdModeStyle> AssetPlacementEdModeStyle;
+	static FName StyleName;
+	static TUniquePtr<FAssetPlacementEdModeStyle> Inst;
 };
