Fix a crash in the asset placement mode when changing active modes.

diff --git a/Engine/Source/Editor/AssetPlacementEdMode/Private/AssetPlacementEdMode.cpp b/Engine/Source/Editor/AssetPlacementEdMode/Private/AssetPlacementEdMode.cpp
--- a/Engine/Source/Editor/AssetPlacementEdMode/Private/AssetPlacementEdMode.cpp
+++ b/Engine/Source/Editor/AssetPlacementEdMode/Private/AssetPlacementEdMode.cpp
@@ -283,7 +283,7 @@ void UAssetPlacementEdMode::ClearSelection()
 
 bool UAssetPlacementEdMode::HasAnyAssetsInPalette() const
 {
-	return (SettingsObjectAsPlacementSettings->PaletteItems.Num() > 0);
+	return SettingsObjectAsPlacementSettings.IsValid() ? (SettingsObjectAsPlacementSettings->PaletteItems.Num() > 0) : false;
 }
 
 bool UAssetPlacementEdMode::HasActiveSelection() const
