Asset Placement mode - fixed foliage placed static meshes not selecting correctly with brush based selection, and not showing up as a valid selection for enabling tools in the mode.

diff --git a/Engine/Source/Editor/AssetPlacementEdMode/Private/AssetPlacementEdMode.cpp b/Engine/Source/Editor/AssetPlacementEdMode/Private/AssetPlacementEdMode.cpp
--- a/Engine/Source/Editor/AssetPlacementEdMode/Private/AssetPlacementEdMode.cpp
+++ b/Engine/Source/Editor/AssetPlacementEdMode/Private/AssetPlacementEdMode.cpp
@@ -288,7 +288,35 @@ bool UAssetPlacementEdMode::HasAnyAssetsInPalette() const
 
 bool UAssetPlacementEdMode::HasActiveSelection() const
 {
-	return HasAnyAssetsInPalette() && Owner->GetEditorSelectionSet()->HasSelectedElements();
+	if (!HasAnyAssetsInPalette())
+	{
+		return false;
+	}
+
+	if (Owner->GetEditorSelectionSet()->HasSelectedElements())
+	{
+		return true;
+	}
+
+	for (TActorIterator<AInstancedFoliageActor> It(GetWorld()); It; ++It)
+	{
+		if (AInstancedFoliageActor* FoliageActor = *It)
+		{
+			bool bHasSelectedFoliage = false;
+			FoliageActor->ForEachFoliageInfo([&bHasSelectedFoliage](UFoliageType* FoliageType, FFoliageInfo& FoliageInfo)
+				{
+					bHasSelectedFoliage = (FoliageInfo.SelectedIndices.Num() > 0);
+					return !bHasSelectedFoliage;
+				});
+
+			if (bHasSelectedFoliage)
+			{
+				return true;
+			}
+		}
+	}
+
+	return false;
 }
 
 bool UAssetPlacementEdMode::IsInSelectionTool() const
diff --git a/Engine/Source/Editor/AssetPlacementEdMode/Private/Tools/PlacementLassoSelectTool.cpp b/Engine/Source/Editor/AssetPlacementEdMode/Private/Tools/PlacementLassoSelectTool.cpp
--- a/Engine/Source/Editor/AssetPlacementEdMode/Private/Tools/PlacementLassoSelectTool.cpp
+++ b/Engine/Source/Editor/AssetPlacementEdMode/Private/Tools/PlacementLassoSelectTool.cpp
@@ -67,6 +67,8 @@ void UPlacementModeLassoSelectTool::OnTick(float DeltaTime)
 					}
 					return true;
 				});
+
+				continue;
 			}
 		}
 
