Asset Placement mode and its tools use a context object from ITF to access typed element common actions and the selection set.

diff --git a/Engine/Source/Editor/AssetPlacementEdMode/Private/AssetPlacementEdMode.cpp b/Engine/Source/Editor/AssetPlacementEdMode/Private/AssetPlacementEdMode.cpp
--- a/Engine/Source/Editor/AssetPlacementEdMode/Private/AssetPlacementEdMode.cpp
+++ b/Engine/Source/Editor/AssetPlacementEdMode/Private/AssetPlacementEdMode.cpp
@@ -12,6 +12,7 @@
 #include "Selection.h"
 #include "EngineUtils.h"
 
+#include "Tools/AssetEditorContextInterface.h"
 #include "Tools/PlacementSelectTool.h"
 #include "Tools/PlacementLassoSelectTool.h"
 #include "Tools/PlacementPlaceTool.h"
@@ -21,14 +22,12 @@
 #include "Factories/AssetFactoryInterface.h"
 #include "Elements/Actor/ActorElementData.h"
 #include "Elements/Framework/EngineElementsLibrary.h"
+#include "Elements/Framework/TypedElementCommonActions.h"
 #include "Elements/Framework/TypedElementRegistry.h"
-#include "Elements/Interfaces/TypedElementAssetDataInterface.h"
-#include "Elements/Interfaces/TypedElementObjectInterface.h"
 
 #include "Settings/LevelEditorMiscSettings.h"
 #include "Modes/PlacementModeSubsystem.h"
 
-#include "Subsystems/EditorActorSubsystem.h"
 #include "InstancedFoliageActor.h"
 
 #define LOCTEXT_NAMESPACE "AssetPlacementEdMode"
@@ -215,31 +214,9 @@ void UAssetPlacementEdMode::DeleteSelection()
 {
 	GetToolManager()->BeginUndoTransaction(LOCTEXT("PlacementDeleteAllSelected", "Delete Selected Assets"));
 
-	// Todo - replace with delete in world interface and replace foliage with element handles
-	UEditorActorSubsystem* ActorSubsystem = GEditor->GetEditorSubsystem<UEditorActorSubsystem>();
-	auto DeleteSelectedElement = [ActorSubsystem](const TTypedElement<UTypedElementObjectInterface>& InElementInterface)
+	if (UTypedElementCommonActions* CommonActions = Owner->GetToolkitHost()->GetCommonActions())
 	{
-		if (InElementInterface.IsSet())
-		{
-			if (AActor* Actor = InElementInterface.GetObjectAs<AActor>())
-			{
-				if (ActorSubsystem)
-				{
-					ActorSubsystem->DestroyActor(Actor);
-				}
-			}
-		}
-	};
-
-	// Gather a copy of the selected handles, since the delete operation will remove them from the selection set.
-	TArray<FTypedElementHandle> SelectedElementHandles;
-	Owner->GetEditorSelectionSet()->GetSelectedElementHandles(SelectedElementHandles, UTypedElementObjectInterface::StaticClass());
-	for (const FTypedElementHandle& ElementHandle : SelectedElementHandles)
-	{
-		if (TTypedElement<UTypedElementObjectInterface> ObjectInterface = UTypedElementRegistry::GetInstance()->GetElement<UTypedElementObjectInterface>(ElementHandle))
-		{
-			DeleteSelectedElement(ObjectInterface);
-		}
+		CommonActions->DeleteSelectedElements(Owner->GetEditorSelectionSet(), GetWorld(), FTypedElementDeletionOptions());
 	}
 
 	for (TActorIterator<AInstancedFoliageActor> It(GetWorld()); It; ++It)
diff --git a/Engine/Source/Editor/AssetPlacementEdMode/Private/Tools/PlacementEraseTool.cpp b/Engine/Source/Editor/AssetPlacementEdMode/Private/Tools/PlacementEraseTool.cpp
--- a/Engine/Source/Editor/AssetPlacementEdMode/Private/Tools/PlacementEraseTool.cpp
+++ b/Engine/Source/Editor/AssetPlacementEdMode/Private/Tools/PlacementEraseTool.cpp
@@ -4,16 +4,21 @@
 #include "UObject/Object.h"
 #include "ToolContextInterfaces.h"
 #include "InteractiveToolManager.h"
-#include "Subsystems/EditorActorSubsystem.h"
 #include "Elements/Framework/EngineElementsLibrary.h"
+#include "Elements/Framework/TypedElementCommonActions.h"
 #include "Elements/Framework/TypedElementHandle.h"
 #include "Elements/Framework/TypedElementRegistry.h"
 #include "Elements/Interfaces/TypedElementObjectInterface.h"
+#include "Tools/AssetEditorContextInterface.h"
 #include "Editor.h"
 #include "InstancedFoliageActor.h"
 #include "AssetPlacementEdMode.h"
 #include "AssetPlacementSettings.h"
 #include "Modes/PlacementModeSubsystem.h"
+#include "Toolkits/IToolkitHost.h"
+#include "UObject/GCObjectScopeGuard.h"
+#include "EditorModeManager.h"
+#include "ContextObjectStore.h"
 
 constexpr TCHAR UPlacementModeEraseTool::ToolName[];
 
@@ -43,20 +48,27 @@ void UPlacementModeEraseTool::OnTick(float DeltaTime)
 		return;
 	}
 
-	UEditorActorSubsystem* ActorSubsystem = GEditor->GetEditorSubsystem<UEditorActorSubsystem>();
+	IAssetEditorContextInterface* AssetEditorContext = GetToolManager()->GetContextObjectStore()->FindContext<IAssetEditorContextInterface>();
+	if (!AssetEditorContext)
+	{
+		return;
+	}
+
+	UTypedElementCommonActions* ElementCommonActions = AssetEditorContext->GetCommonActions();
+	if (!ElementCommonActions)
+	{
+		return;
+	}
+
+	TGCObjectScopeGuard<UTypedElementList> ElementsToDelete(UTypedElementRegistry::GetInstance()->CreateElementList());
+
 	TArray<FTypedElementHandle> HitElements = GetElementsInBrushRadius();
 	for (const FTypedElementHandle& HitElement : HitElements)
 	{
 		if (TTypedElement<UTypedElementObjectInterface> ObjectInterface = UTypedElementRegistry::GetInstance()->GetElement<UTypedElementObjectInterface>(HitElement))
 		{
-			AActor* Actor = ObjectInterface.GetObjectAs<AActor>();
-			if (!Actor)
-			{
-				continue;
-			}
-
 			// Since the foliage static mesh instances do not currently operate with element handles, we have to drill in manually here.
-			if (AInstancedFoliageActor* FoliageActor = Cast<AInstancedFoliageActor>(Actor))
+			if (AInstancedFoliageActor* FoliageActor = ObjectInterface.GetObjectAs<AInstancedFoliageActor>())
 			{
 				FoliageActor->ForEachFoliageInfo([this](UFoliageType* FoliageType, FFoliageInfo& FoliageInfo)
 				{
@@ -68,13 +80,22 @@ void UPlacementModeEraseTool::OnTick(float DeltaTime)
 						FoliageInfo.GetInstancesInsideSphere(SphereToCheck, Instances);
 						FoliageInfo.RemoveInstances(Instances, true);
 					}
-					return true; // continue iteraton
+					return true; // continue iteration
 				});
 			}
-			else if (ActorSubsystem)
+			else
 			{
-				ActorSubsystem->DestroyActor(Actor);
+				ElementsToDelete.Get()->Add(HitElement);
 			}
 		}
 	}
+
+	if (ElementsToDelete.Get()->HasElements())
+	{
+		UTypedElementSelectionSet* SelectionSet = AssetEditorContext->GetMutableSelectionSet();
+		if (SelectionSet)
+		{
+			ElementCommonActions->DeleteElementsInList(ElementsToDelete.Get(), AssetEditorContext->GetEditingWorld(), SelectionSet, FTypedElementDeletionOptions());
+		}
+	}
 }
diff --git a/Engine/Source/Editor/AssetPlacementEdMode/Private/Tools/PlacementLassoSelectTool.cpp b/Engine/Source/Editor/AssetPlacementEdMode/Private/Tools/PlacementLassoSelectTool.cpp
--- a/Engine/Source/Editor/AssetPlacementEdMode/Private/Tools/PlacementLassoSelectTool.cpp
+++ b/Engine/Source/Editor/AssetPlacementEdMode/Private/Tools/PlacementLassoSelectTool.cpp
@@ -16,6 +16,9 @@
 #include "BaseBehaviors/KeyAsModifierInputBehavior.h"
 #include "Elements/Framework/TypedElementSelectionSet.h"
 #include "Modes/PlacementModeSubsystem.h"
+#include "Tools/AssetEditorContextInterface.h"
+#include "EditorModeManager.h"
+#include "ContextObjectStore.h"
 
 constexpr TCHAR UPlacementModeLassoSelectTool::ToolName[];
 
@@ -45,6 +48,18 @@ void UPlacementModeLassoSelectTool::OnTick(float DeltaTime)
 		return;
 	}
 
+	IAssetEditorContextInterface* AssetEditorContext = GetToolManager()->GetContextObjectStore()->FindContext<IAssetEditorContextInterface>();
+	if (!AssetEditorContext)
+	{
+		return;
+	}
+
+	UTypedElementSelectionSet* SelectionSet = AssetEditorContext->GetMutableSelectionSet();
+	if (!SelectionSet)
+	{
+		return;
+	}
+
 	bool bSelectElements = !bCtrlToggle;
 
 	TArray<FTypedElementHandle> HitElements = GetElementsInBrushRadius();
@@ -72,18 +87,13 @@ void UPlacementModeLassoSelectTool::OnTick(float DeltaTime)
 			}
 		}
 
-		FToolBuilderState SelectionState;
-		GetToolManager()->GetContextQueriesAPI()->GetCurrentSelectionState(SelectionState);
-		if (SelectionState.TypedElementSelectionSet.IsValid())
+		if (bSelectElements)
 		{
-			if (bSelectElements)
-			{
-				SelectionState.TypedElementSelectionSet->SelectElement(HitElement, FTypedElementSelectionOptions());
-			}
-			else
-			{
-				SelectionState.TypedElementSelectionSet->DeselectElement(HitElement, FTypedElementSelectionOptions());
-			}
+			SelectionSet->SelectElement(HitElement, FTypedElementSelectionOptions());
+		}
+		else
+		{
+			SelectionSet->DeselectElement(HitElement, FTypedElementSelectionOptions());
 		}
 	}
 }
diff --git a/Engine/Source/Editor/AssetPlacementEdMode/Private/Tools/PlacementPlaceSingleTool.cpp b/Engine/Source/Editor/AssetPlacementEdMode/Private/Tools/PlacementPlaceSingleTool.cpp
--- a/Engine/Source/Editor/AssetPlacementEdMode/Private/Tools/PlacementPlaceSingleTool.cpp
+++ b/Engine/Source/Editor/AssetPlacementEdMode/Private/Tools/PlacementPlaceSingleTool.cpp
@@ -19,6 +19,10 @@
 #include "Modes/PlacementModeSubsystem.h"
 #include "Subsystems/PlacementSubsystem.h"
 #include "UObject/Object.h"
+#include "Tools/AssetEditorContextInterface.h"
+#include "EditorModeManager.h"
+#include "Toolkits/IToolkitHost.h"
+#include "ContextObjectStore.h"
 
 constexpr TCHAR UPlacementModePlaceSingleTool::ToolName[];
 
@@ -310,19 +314,27 @@ void UPlacementModePlaceSingleTool::EnterTweakState(TArrayView<const FTypedEleme
 		return;
 	}
 
-	FToolBuilderState SelectionState;
-	GetToolManager()->GetContextQueriesAPI()->GetCurrentSelectionState(SelectionState);
-	SelectionState.TypedElementSelectionSet->SetSelection(InElementHandles, FTypedElementSelectionOptions());
-	bIsTweaking = true;
+	if (IAssetEditorContextInterface* AssetEditorContext = GetToolManager()->GetContextObjectStore()->FindContext<IAssetEditorContextInterface>())
+	{
+		if (UTypedElementSelectionSet* SelectionSet = AssetEditorContext->GetMutableSelectionSet())
+		{
+			SelectionSet->SetSelection(InElementHandles, FTypedElementSelectionOptions());
+			bIsTweaking = true;
+		}
+	}
 }
 
 void UPlacementModePlaceSingleTool::ExitTweakState(bool bClearSelectionSet)
 {
 	if (bClearSelectionSet)
 	{
-		FToolBuilderState SelectionState;
-		GetToolManager()->GetContextQueriesAPI()->GetCurrentSelectionState(SelectionState);
-		SelectionState.TypedElementSelectionSet->ClearSelection(FTypedElementSelectionOptions());
+		if (IAssetEditorContextInterface* AssetEditorContext = GetToolManager()->GetContextObjectStore()->FindContext<IAssetEditorContextInterface>())
+		{
+			if (UTypedElementSelectionSet* SelectionSet = AssetEditorContext->GetMutableSelectionSet())
+			{
+				SelectionSet->ClearSelection(FTypedElementSelectionOptions());
+			}
+		}
 	}
 
 	bIsTweaking = false;
@@ -356,16 +368,24 @@ void UPlacementModePlaceSingleTool::NotifyMovementStarted(TArrayView<const FType
 
 void UPlacementModePlaceSingleTool::NotifyMovementEnded(TArrayView<const FTypedElementHandle> InElements)
 {
-	FToolBuilderState SelectionState;
-	GetToolManager()->GetContextQueriesAPI()->GetCurrentSelectionState(SelectionState);
-	check(SelectionState.TypedElementSelectionSet.IsValid());	// Placement tools expect a valid selection set.
+	IAssetEditorContextInterface* AssetEditorContext = GetToolManager()->GetContextObjectStore()->FindContext<IAssetEditorContextInterface>();
+	if (!AssetEditorContext)
+	{
+		return;
+	}
+
+	UTypedElementSelectionSet* SelectionSet = AssetEditorContext->GetMutableSelectionSet();
+	if (!SelectionSet)
+	{
+		return;
+	}
 
 	for (const FTypedElementHandle& PreviewElement : PreviewElements)
 	{
 		if (TTypedElement<UTypedElementWorldInterface> WorldInterfaceElement = UTypedElementRegistry::GetInstance()->GetElement<UTypedElementWorldInterface>(PreviewElement))
 		{
 			WorldInterfaceElement.NotifyMovementEnded();
-			WorldInterfaceElement.DeleteElement(WorldInterfaceElement.GetOwnerWorld(), SelectionState.TypedElementSelectionSet.Get(), FTypedElementDeletionOptions());
+			WorldInterfaceElement.DeleteElement(WorldInterfaceElement.GetOwnerWorld(), SelectionSet, FTypedElementDeletionOptions());
 		}
 	}
 }
diff --git a/Engine/Source/Editor/EditorFramework/Public/Tools/AssetEditorContextInterface.h b/Engine/Source/Editor/EditorFramework/Public/Tools/AssetEditorContextInterface.h
new file mode 100644
--- /dev/null
+++ b/Engine/Source/Editor/EditorFramework/Public/Tools/AssetEditorContextInterface.h
@@ -0,0 +1,46 @@
+// Copyright Epic Games, Inc. All Rights Reserved.
+
+#pragma once
+
+#include "UObject/Interface.h"
+
+#include "AssetEditorContextInterface.generated.h"
+
+class IToolkitHost;
+class UTypedElementCommonActions;
+class UTypedElementSelectionSet;
+
+UINTERFACE(Blueprintable, MinimalAPI)
+class UAssetEditorContextInterface : public UInterface
+{
+	GENERATED_BODY()
+};
+
+/**
+ * An interface for a context object that exposes common functionality of an asset editor.
+ */
+class EDITORFRAMEWORK_API IAssetEditorContextInterface
+{	
+	GENERATED_BODY()
+
+public:
+	/**
+	 * Get the selection set that the asset editor uses.
+	 */
+	virtual const UTypedElementSelectionSet* GetSelectionSet() const = 0;
+
+	/**
+	 * Get the selection set that the asset editor uses.
+	 */
+	virtual UTypedElementSelectionSet* GetMutableSelectionSet() = 0;
+
+	/**
+	 * Get the typed element common actions for the asset editor.
+	 */
+	virtual UTypedElementCommonActions* GetCommonActions() = 0;
+
+	/**
+	 * Gets the world on which the asset editor performs actions.
+	 */
+	virtual UWorld* GetEditingWorld() const = 0;
+};
diff --git a/Engine/Source/Editor/UnrealEd/Private/EditorModeManager.cpp b/Engine/Source/Editor/UnrealEd/Private/EditorModeManager.cpp
--- a/Engine/Source/Editor/UnrealEd/Private/EditorModeManager.cpp
+++ b/Engine/Source/Editor/UnrealEd/Private/EditorModeManager.cpp
@@ -40,6 +40,8 @@
 #include "CanvasItem.h"
 #include "Engine/StaticMeshActor.h"
 #include "EngineUtils.h"
+#include "Tools/AssetEditorContextObject.h"
+#include "ContextObjectStore.h"
 
 /*------------------------------------------------------------------------------
 	FEditorModeTools.
@@ -152,6 +154,13 @@ void FEditorModeTools::SetToolkitHost(TSharedRef<class IToolkitHost> InHost)
 {
 	checkf(!ToolkitHost.IsValid(), TEXT("SetToolkitHost can only be called once"));
 	ToolkitHost = InHost;
+
+	if (HasToolkitHost())
+	{
+		UAssetEditorContextObject* AssetEditorContextObject = NewObject<UAssetEditorContextObject>(InteractiveToolsContext->ToolManager);
+		AssetEditorContextObject->SetToolkitHost(GetToolkitHost().Get());
+		InteractiveToolsContext->ContextObjectStore->AddContextObject(AssetEditorContextObject);
+	}
 }
 
 USelection* FEditorModeTools::GetSelectedActors() const
diff --git a/Engine/Source/Editor/UnrealEd/Private/Tools/AssetEditorContextObject.cpp b/Engine/Source/Editor/UnrealEd/Private/Tools/AssetEditorContextObject.cpp
new file mode 100644
--- /dev/null
+++ b/Engine/Source/Editor/UnrealEd/Private/Tools/AssetEditorContextObject.cpp
@@ -0,0 +1,53 @@
+// Copyright Epic Games, Inc. All Rights Reserved.
+
+#include "Tools/AssetEditorContextObject.h"
+
+#include "EditorModeManager.h"
+#include "Elements/Framework/TypedElementSelectionSet.h"
+#include "Elements/Framework/TypedElementCommonActions.h"
+#include "Toolkits/IToolkitHost.h"
+
+const UTypedElementSelectionSet* UAssetEditorContextObject::GetSelectionSet() const
+{
+	if (GetToolkitHost())
+	{
+		return GetToolkitHost()->GetEditorModeManager().GetEditorSelectionSet();
+	}
+
+	return nullptr;
+}
+
+UTypedElementSelectionSet* UAssetEditorContextObject::GetMutableSelectionSet()
+{
+	if (GetToolkitHost())
+	{
+		return GetToolkitHost()->GetEditorModeManager().GetEditorSelectionSet();
+	}
+
+	return nullptr;
+}
+
+UTypedElementCommonActions* UAssetEditorContextObject::GetCommonActions()
+{
+	if (GetToolkitHost())
+	{
+		return GetToolkitHost()->GetCommonActions();
+	}
+
+	return nullptr;
+}
+
+UWorld* UAssetEditorContextObject::GetEditingWorld() const
+{
+	if (GetToolkitHost())
+	{
+		return GetToolkitHost()->GetWorld();
+	}
+
+	return nullptr;
+}
+
+UWorld* UAssetEditorContextObject::GetWorld() const
+{
+	return GetEditingWorld();
+}
diff --git a/Engine/Source/Editor/UnrealEd/Private/Tools/AssetEditorContextObject.h b/Engine/Source/Editor/UnrealEd/Private/Tools/AssetEditorContextObject.h
new file mode 100644
--- /dev/null
+++ b/Engine/Source/Editor/UnrealEd/Private/Tools/AssetEditorContextObject.h
@@ -0,0 +1,27 @@
+// Copyright Epic Games, Inc. All Rights Reserved.
+
+#pragma once
+
+#include "Elements/Framework/TypedElementAssetEditorToolkitHostMixin.h"
+#include "Tools/AssetEditorContextInterface.h"
+#include "UObject/Object.h"
+
+#include "AssetEditorContextObject.generated.h"
+
+UCLASS(Transient)
+class UAssetEditorContextObject : public UObject, public IAssetEditorContextInterface, public FTypedElementAssetEditorToolkitHostMixin
+{
+	GENERATED_BODY()
+
+public:
+	// Begin UAssetEditorContextObject interface
+	const UTypedElementSelectionSet* GetSelectionSet() const override;
+	UTypedElementSelectionSet* GetMutableSelectionSet() override;
+	UTypedElementCommonActions* GetCommonActions() override;
+	UWorld* GetEditingWorld() const override;
+	// End UAssetEditorContextObject interface
+
+	// Begin UObject interface
+	UWorld* GetWorld() const override;
+	// End UObject interface
+};
\ No newline at end of file
