Updated the placement palette to use SAssetDropTarget, so that it will get UX styling improvements to asset drag/drop hover state.

#b jamie.dale

diff --git a/Engine/Source/Editor/AssetPlacementEdMode/AssetPlacementEdMode.Build.cs b/Engine/Source/Editor/AssetPlacementEdMode/AssetPlacementEdMode.Build.cs
--- a/Engine/Source/Editor/AssetPlacementEdMode/AssetPlacementEdMode.Build.cs
+++ b/Engine/Source/Editor/AssetPlacementEdMode/AssetPlacementEdMode.Build.cs
@@ -26,6 +26,7 @@ public class AssetPlacementEdMode : ModuleRules
 				"TypedElementFramework",
 				"TypedElementRuntime",
 				"PropertyEditor",
+				"EditorWidgets",
 			}
 		);
 	}
diff --git a/Engine/Source/Editor/AssetPlacementEdMode/Private/SAssetPlacementPalette.cpp b/Engine/Source/Editor/AssetPlacementEdMode/Private/SAssetPlacementPalette.cpp
--- a/Engine/Source/Editor/AssetPlacementEdMode/Private/SAssetPlacementPalette.cpp
+++ b/Engine/Source/Editor/AssetPlacementEdMode/Private/SAssetPlacementPalette.cpp
@@ -36,71 +36,10 @@
 #include "IContentBrowserSingleton.h"
 #include "ContentBrowserModule.h"
 #include "Modules/ModuleManager.h"
+#include "SAssetDropTarget.h"
 
 #define LOCTEXT_NAMESPACE "AssetPlacementMode"
 
-////////////////////////////////////////////////
-// SPlacementDragDropHandler
-////////////////////////////////////////////////
-
-/** Drag-drop zone for adding Placement types to the palette */
-class SAssetPaletteDragDropHandler : public SCompoundWidget
-{
-public:
-	SLATE_BEGIN_ARGS(SAssetPaletteDragDropHandler) {}
-		SLATE_DEFAULT_SLOT(FArguments, Content)
-		SLATE_EVENT(FOnDrop, OnDrop)
-	SLATE_END_ARGS()
-
-	void Construct(const FArguments& InArgs)
-	{
-		bIsDragOn = false;
-		OnDropDelegate = InArgs._OnDrop;
-
-		this->ChildSlot
-			[
-				SNew(SBorder)
-				.BorderImage(FAppStyle::Get().GetBrush("WhiteBrush"))
-				.BorderBackgroundColor(this, &SAssetPaletteDragDropHandler::GetBackgroundColor)
-				.Padding(FMargin(100))
-				[
-					InArgs._Content.Widget
-				]
-			];
-	}
-
-	FReply OnDrop(const FGeometry& MyGeometry, const FDragDropEvent& DragDropEvent) override
-	{
-		bIsDragOn = false;
-		if (OnDropDelegate.IsBound())
-		{
-			return OnDropDelegate.Execute(MyGeometry, DragDropEvent);
-		}
-		
-		return FReply::Handled();
-	}
-
-	virtual void OnDragEnter(const FGeometry& MyGeometry, const FDragDropEvent& DragDropEvent) override
-	{
-		bIsDragOn = true;
-	}
-
-	virtual void OnDragLeave(const FDragDropEvent& DragDropEvent) override
-	{
-		bIsDragOn = false;
-	}
-
-private:
-	FSlateColor GetBackgroundColor() const
-	{
-		return bIsDragOn ? FLinearColor(1.0f, 0.6f, 0.1f, 0.9f) : FLinearColor(0.1f, 0.1f, 0.1f, 0.9f);
-	}
-
-private:
-	FOnDrop OnDropDelegate;
-	bool bIsDragOn;
-};
-
 ////////////////////////////////////////////////
 // SPlacementPalette
 ////////////////////////////////////////////////
@@ -211,22 +150,6 @@ void SAssetPlacementPalette::Construct(const FArguments& InArgs)
 			[
 				SNew(SVerticalBox)
 
-				+ SVerticalBox::Slot()
-				[
-					SNew(SBox)
-					.Visibility(this, &SAssetPlacementPalette::GetDropPlacementHintVisibility)
-					.MinDesiredHeight(100)
-					[
-						SNew(SScaleBox)
-						.Stretch(EStretch::ScaleToFit)
-						[
-							SNew(STextBlock)
-							.Text(LOCTEXT("Placement_DropStatic", "+ Drop Assets Here"))
-							.ToolTipText(LOCTEXT("Placement_DropStatic_ToolTip", "Drag and drop asset types from the Content Browser to add them to the palette."))
-						]
-					]
-				]
-
 				+ SVerticalBox::Slot()
 				.Padding(2.f, 0.f)
 				[
@@ -234,20 +157,23 @@ void SAssetPlacementPalette::Construct(const FArguments& InArgs)
 				]
 			]
 			
-			// Placement Mesh Drop Zone
+			// Placement Drop Zone
 			+ SOverlay::Slot()
 			.HAlign(HAlign_Fill)
 			.VAlign(VAlign_Fill)
 			[
-				SNew(SAssetPaletteDragDropHandler)
-				.Visibility(this, &SAssetPlacementPalette::GetPlacementDropTargetVisibility)
-				.OnDrop(this, &SAssetPlacementPalette::HandlePlacementDropped)
+				SNew(SAssetDropTarget)
+				.OnAreAssetsAcceptableForDrop(this, &SAssetPlacementPalette::OnAreAssetsValidForDrop)
+				.OnAssetsDropped(this, &SAssetPlacementPalette::HandlePlacementDropped)
+				.bSupportsMultiDrop(true)
 				[
-					SNew(SScaleBox)
-					.Stretch(EStretch::ScaleToFit)
+					SNew(SBox)
+					.HAlign(HAlign_Center)
+					.VAlign(VAlign_Center)
 					[
 						SNew(STextBlock)
-						.Text(LOCTEXT("Placement_AddPlacementMesh", "+ Asset Type"))
+						.Text(LOCTEXT("Placement_AddPlacementMesh", "Drop Assets Here"))
+						.Visibility(this, &SAssetPlacementPalette::GetDropHintTextVisibility)
 						.ShadowOffset(FVector2D(1.f, 1.f))
 					]
 				]
@@ -597,67 +523,58 @@ TSharedPtr<SListView<FPlacementPaletteItemModelPtr>> SAssetPlacementPalette::Get
 	return nullptr;
 }
 
-EVisibility SAssetPlacementPalette::GetDropPlacementHintVisibility() const
+EVisibility SAssetPlacementPalette::GetDropHintTextVisibility() const
 {
-	if (bIsMirroringContentBrowser)
+	if (!bIsMirroringContentBrowser && (FilteredItems.Num() == 0) && !FSlateApplication::Get().IsDragDropping())
 	{
-		return EVisibility::Collapsed;
+		return EVisibility::HitTestInvisible;
 	}
 
-	return FilteredItems.Num() ? EVisibility::Collapsed : EVisibility::Visible;
+	return EVisibility::Hidden;
 }
 
-EVisibility SAssetPlacementPalette::GetPlacementDropTargetVisibility() const
+bool SAssetPlacementPalette::OnAreAssetsValidForDrop(TArrayView<FAssetData> DraggedAssets) const
 {
 	if (bIsMirroringContentBrowser)
 	{
-		return EVisibility::Hidden;
+		return false;
 	}
 
-	if (FSlateApplication::Get().IsDragDropping())
+	UPlacementSubsystem* PlacementSubystem = GEditor->GetEditorSubsystem<UPlacementSubsystem>();
+	for (const FAssetData& AssetData : DraggedAssets)
 	{
-		TArray<FAssetData> DraggedAssets = AssetUtil::ExtractAssetDataFromDrag(FSlateApplication::Get().GetDragDroppingContent());
-		for (const FAssetData& AssetData : DraggedAssets)
+		if (AssetData.GetClass()->HasAnyClassFlags(CLASS_Abstract | CLASS_Deprecated | CLASS_NewerVersionExists | CLASS_NotPlaceable))
 		{
-			if (AssetData.GetClass()->HasAnyClassFlags(CLASS_Abstract | CLASS_Deprecated | CLASS_NewerVersionExists | CLASS_NotPlaceable))
-			{
-				continue;
-			}
+			continue;
+		}
 
-			if (AssetData.IsValid())
-			{
-				return EVisibility::Visible;
-			}
+		if (AssetData.IsValid())
+		{
+			return PlacementSubystem ? (PlacementSubystem->FindAssetFactoryFromAssetData(AssetData).GetInterface() != nullptr) : true;
 		}
 	}
 
-	return EVisibility::Hidden;
+	return false;
 }
 
-FReply SAssetPlacementPalette::HandlePlacementDropped(const FGeometry& DropZoneGeometry, const FDragDropEvent& DragDropEvent)
+void SAssetPlacementPalette::HandlePlacementDropped(const FDragDropEvent& DragDropEvent, TArrayView<FAssetData> DroppedAssetData)
 {
-	if (bIsMirroringContentBrowser)
+	if ((DroppedAssetData.Num() == 0) || bIsMirroringContentBrowser)
 	{
-		return FReply::Unhandled();
+		return;
 	}
 
-	TArray<FAssetData> DroppedAssetData = AssetUtil::ExtractAssetDataFromDrag(DragDropEvent);
-	if (DroppedAssetData.Num() > 0)
+	if (DragDropEvent.IsShiftDown())
 	{
-		if (DragDropEvent.IsShiftDown())
-		{
-			ClearPalette();
-		}
-
-		for (auto& AssetData : DroppedAssetData)
-		{
-			AddPlacementType(AssetData);
-		}
+		ClearPalette();
+	}
 
-		UpdatePalette(true);
+	for (auto& AssetData : DroppedAssetData)
+	{
+		AddPlacementType(AssetData);
 	}
 
-	return FReply::Handled();
+	UpdatePalette(true);
 }
 
 bool SAssetPlacementPalette::HasAnyItemInPalette() const
diff --git a/Engine/Source/Editor/AssetPlacementEdMode/Private/SAssetPlacementPalette.h b/Engine/Source/Editor/AssetPlacementEdMode/Private/SAssetPlacementPalette.h
--- a/Engine/Source/Editor/AssetPlacementEdMode/Private/SAssetPlacementPalette.h
+++ b/Engine/Source/Editor/AssetPlacementEdMode/Private/SAssetPlacementPalette.h
@@ -22,10 +22,10 @@ class FAssetPlacementPaletteItemModel;
 class FMenuBuilder;
 class FUICommandList;
 class IDetailsView;
-class UFoliageType;
 struct FAssetData;
 class UAssetPlacementSettings;
 class IPropertyHandle;
+class SSearchBox;
 
 typedef TSharedPtr<FAssetPlacementPaletteItemModel> FPlacementPaletteItemModelPtr;
 typedef STreeView<FPlacementPaletteItemModelPtr> SPlacementTypeTreeView;
@@ -102,25 +102,18 @@ private:
 	/** Sets whether to show tooltips when hovering over Placement type items in the palette */
 	void ToggleShowTooltips();
 
-	/** Switches the palette display between the tile and tree view */
-	FReply OnToggleViewModeClicked();
-
-	/** @return The index of the view widget to display */
-	int32 GetActiveViewIndex() const;
-
 	/** Creates the view options menu */
 	TSharedRef<SWidget> GetViewOptionsMenuContent();
 
 	TSharedPtr<SListView<FPlacementPaletteItemModelPtr>> GetActiveViewWidget() const;
 
-	/** Gets the visibility of the "Drop Placement Here" prompt for when the palette is empty */
-	EVisibility GetDropPlacementHintVisibility() const;
+	/** Gets the visibility of the drop hint text for when the palette or filtered palette is empty */
+	EVisibility GetDropHintTextVisibility() const;
 
-	/** Gets the visibility of the drag-drop zone overlay */
-	EVisibility GetPlacementDropTargetVisibility() const;
+	bool OnAreAssetsValidForDrop(TArrayView<FAssetData> DraggedAssets) const;
 
 	/** Handles dropping of a mesh or Placement type into the palette */
-	FReply HandlePlacementDropped(const FGeometry& DropZoneGeometry, const FDragDropEvent& DragDropEvent);
+	void HandlePlacementDropped(const FDragDropEvent& DragDropEvent, TArrayView<FAssetData> DraggedAssets);
 
 	/** @returns true if there are any items in the palette. */
 	bool HasAnyItemInPalette() const;
@@ -163,10 +156,6 @@ private:	// TREE VIEW
 	EColumnSortMode::Type GetMeshColumnSortMode() const;
 	void OnTypeColumnSortModeChanged(EColumnSortPriority::Type InPriority, const FName& InColumnName, EColumnSortMode::Type InSortMode);
 
-private:
-	/** Handles the click for the uneditable blueprint Placement type warning */
-	void OnEditPlacementTypeBlueprintHyperlinkClicked(const FSlateHyperlinkRun::FMetadata& Metadata);
-
 private:
 	/** Active timer handler to update the items in the palette */
 	EActiveTimerReturnType UpdatePaletteItems(double InCurrentTime, float InDeltaTime);
@@ -197,18 +186,18 @@ private:
 	TSharedPtr<SPlacementTypeTreeView> TreeViewWidget;
 
 	/** Placement mesh details widget  */
-	TSharedPtr<class IDetailsView> DetailsWidget;
+	TSharedPtr<IDetailsView> DetailsWidget;
 
 	TSharedPtr<IPropertyHandle> PalettePropertyHandle;
 
 	/** Placement items search box widget */
-	TSharedPtr<class SSearchBox> SearchBoxPtr;
+	TSharedPtr<SSearchBox> SearchBoxPtr;
 
 	/** Command list for binding functions for the context menu. */
 	TSharedPtr<FUICommandList> UICommandList;
 
 	/** Thumbnail pool for rendering mesh thumbnails */
-	TSharedPtr<class FAssetThumbnailPool> ThumbnailPool;
+	TSharedPtr<FAssetThumbnailPool> ThumbnailPool;
 
 	bool bItemsNeedRebuild = true;
 	bool bShowFullTooltips = true;
