Placement tools use hit proxies to determine objects inside of brush radius instead of world sweeps.

Static mesh instances are only hooked up editor side via hit proxies, and won't register in a sweep since it only deals with limited types.

diff --git a/Engine/Source/Editor/AssetPlacementEdMode/Private/Modes/PlacementModeSubsystem.cpp b/Engine/Source/Editor/AssetPlacementEdMode/Private/Modes/PlacementModeSubsystem.cpp
--- a/Engine/Source/Editor/AssetPlacementEdMode/Private/Modes/PlacementModeSubsystem.cpp
+++ b/Engine/Source/Editor/AssetPlacementEdMode/Private/Modes/PlacementModeSubsystem.cpp
@@ -56,7 +56,7 @@ bool UPlacementModeSubsystem::DoesCurrentPaletteSupportElement(const FTypedEleme
 	// Not all actors do this reliably, so additionally check the supplied factory for a match. 
 	for (const TSharedPtr<FPaletteItem>& Item : ModeSettings->PaletteItems)
 	{
-		if (!Item)
+		if (!Item || !Item->AssetFactoryInterface)
 		{
 			continue;
 		}
diff --git a/Engine/Source/Editor/AssetPlacementEdMode/Private/Tools/PlacementBrushToolBase.cpp b/Engine/Source/Editor/AssetPlacementEdMode/Private/Tools/PlacementBrushToolBase.cpp
--- a/Engine/Source/Editor/AssetPlacementEdMode/Private/Tools/PlacementBrushToolBase.cpp
+++ b/Engine/Source/Editor/AssetPlacementEdMode/Private/Tools/PlacementBrushToolBase.cpp
@@ -7,18 +7,21 @@
 #include "InstancedFoliageActor.h"
 #include "FoliageHelper.h"
 #include "Components/PrimitiveComponent.h"
-#include "Components/StaticMeshComponent.h"
 #include "Components/BrushComponent.h"
 #include "Components/ModelComponent.h"
 #include "LandscapeHeightfieldCollisionComponent.h"
 #include "FoliageInstancedStaticMeshComponent.h"
 #include "AssetPlacementSettings.h"
 #include "Elements/Framework/EngineElementsLibrary.h"
-#include "AssetPlacementEdMode.h"
+#include "Elements/Framework/TypedElementRegistry.h"
+#include "Elements/Framework/TypedElementSelectionSet.h"
+#include "Elements/Interfaces/TypedElementWorldInterface.h"
+#include "Elements/Interfaces/TypedElementSelectionInterface.h"
 #include "ActorPartition/ActorPartitionSubsystem.h"
 #include "Editor.h"
 #include "Modes/PlacementModeSubsystem.h"
 #include "ActorFactories/ActorFactory.h"
+#include "BaseGizmos/GizmoRenderingUtil.h"
 #include "Elements/SMInstance/SMInstanceElementData.h" // For SMInstanceElementDataUtil::SMInstanceElementsEnabled
 
 bool UPlacementToolBuilderBase::CanBuildTool(const FToolBuilderState& SceneState) const
@@ -45,6 +48,27 @@ bool UPlacementBrushToolBase::AreAllTargetsValid() const
 	return Target ? Target->IsValid() : true;
 }
 
+void UPlacementBrushToolBase::Render(IToolsContextRenderAPI* RenderAPI)
+{
+	Super::Render(RenderAPI);
+
+	// Transform the brush radius to standard pixel size
+	LastBrushStampWorldToPixelScale = GizmoRenderingUtil::CalculateLocalPixelToWorldScale(RenderAPI->GetSceneView(), LastBrushStamp.WorldPosition);
+}
+
+void UPlacementBrushToolBase::OnClickPress(const FInputDeviceRay& PressPos)
+{
+	LastDeviceInputRay = PressPos;
+	Super::OnClickPress(PressPos);
+}
+
+void UPlacementBrushToolBase::OnClickDrag(const FInputDeviceRay& DragPos)
+{
+	LastDeviceInputRay = DragPos;
+	GetToolManager()->PostInvalidation();
+	Super::OnClickDrag(DragPos);
+}
+
 double UPlacementBrushToolBase::EstimateMaximumTargetDimension()
 {
 	return 1000.0;
@@ -238,52 +262,61 @@ FTransform UPlacementBrushToolBase::FinalizeTransform(const FTransform& Original
 	return FinalizedTransform;
 }
 
-TArray<FTypedElementHandle> UPlacementBrushToolBase::GetElementsInBrushRadius() const
+TArray<FTypedElementHandle> UPlacementBrushToolBase::GetElementsInBrushRadius(const FInputDeviceRay& DragPos) const
 {
 	TArray<FTypedElementHandle> ElementHandles;
-	FCollisionQueryParams QueryParams(TEXT("PlacementBrushTool"), SCENE_QUERY_STAT_ONLY(IFA_FoliageTrace), true);
-	QueryParams.bReturnFaceIndex = false;
-	TArray<FHitResult> Hits;
-	FCollisionShape BrushSphere;
-	BrushSphere.SetSphere(LastBrushStamp.Radius);
 
-	const FVector TraceStart(LastWorldRay.Origin);
-	const FVector TraceEnd(LastWorldRay.Origin + LastWorldRay.Direction * HALF_WORLD_MAX);
-
-	GetToolManager()->GetWorld()->SweepMultiByObjectType(Hits, TraceStart, TraceEnd, FQuat::Identity,
-		FCollisionObjectQueryParams(FCollisionObjectQueryParams::InitType::AllObjects), BrushSphere, QueryParams);
-
-	for (const FHitResult& Hit : Hits)
+	// We need the 2D device screen space position to test against hit proxies.
+	if (!DragPos.bHas2D)
 	{
-		const UPrimitiveComponent* HitComponent = Hit.GetComponent();
-		check(HitComponent);
-
-		// In the editor traces can hit "No Collision" type actors, so ugh. (ignore these)
-		if (!HitComponent->IsQueryCollisionEnabled() || HitComponent->GetCollisionResponseToChannel(ECC_WorldStatic) != ECR_Block)
-		{
-			continue;
-		}
+		return ElementHandles;
+	}
 
-		// Don't place assets on invisible walls / triggers / volumes
-		if (HitComponent->IsA<UBrushComponent>())
-		{
-			continue;
-		}
+	FViewport* Viewport = GetToolManager()->GetContextQueriesAPI()->GetFocusedViewport();
+	if (!Viewport)
+	{
+		return ElementHandles;
+	}
 
-		const AActor* HitActor = Hit.GetActor();
-		if (HitActor)
+	FToolBuilderState SelectionState;
+	GetToolManager()->GetContextQueriesAPI()->GetCurrentSelectionState(SelectionState);
+	if (!SelectionState.TypedElementSelectionSet.IsValid())
+	{
+		return ElementHandles;
+	}
+	UTypedElementSelectionSet* SelectionSet = SelectionState.TypedElementSelectionSet.Get();
+
+	// Convert brush radius to screen space and gather hit elements within the brush radius.
+	int32 HalfBrushRadius = FMath::CeilToFloat((LastBrushStamp.Radius * LastBrushStampWorldToPixelScale) / 2.0f);
+	FIntRect AreaToCheck;
+	FIntPoint ViewportSize = Viewport->GetSizeXY();
+	AreaToCheck.Min.X = FMath::Max<int32>(0, DragPos.ScreenPosition.X - HalfBrushRadius);
+	AreaToCheck.Min.Y = FMath::Max<int32>(0, DragPos.ScreenPosition.Y - HalfBrushRadius);
+	AreaToCheck.Max.X = FMath::Min<int32>(ViewportSize.X, DragPos.ScreenPosition.X + HalfBrushRadius);
+	AreaToCheck.Max.Y = FMath::Min<int32>(ViewportSize.Y, DragPos.ScreenPosition.Y + HalfBrushRadius);
+
+	FBoxSphereBounds WorldBrushSphereBounds(FSphere(LastBrushStamp.WorldPosition, LastBrushStamp.Radius));
+	TSet<FTypedElementHandle> HitElementHandles;
+	Viewport->GetElementHandlesInRect(AreaToCheck, HitElementHandles);
+	for (const FTypedElementHandle& HitHandle : HitElementHandles)
+	{
+		FTypedElementHandle ResolvedHandle = SelectionSet->GetSelectionElement(HitHandle, ETypedElementSelectionMethod::Primary);
+		if (TTypedElement<UTypedElementWorldInterface> WorldInterface = UTypedElementRegistry::GetInstance()->GetElement<UTypedElementWorldInterface>(ResolvedHandle))
 		{
-			FTypedElementHandle ActorHandle = UEngineElementsLibrary::AcquireEditorActorElementHandle(HitActor);
-			if (GEditor->GetEditorSubsystem<UPlacementModeSubsystem>()->DoesCurrentPaletteSupportElement(ActorHandle))
+			// Since the viewport gives us back hit proxies in a rect, we also want to verify that we're actually intersecting with the sphere from the brush in world space.
+			FBoxSphereBounds ElementBounds;
+			WorldInterface.GetBounds(ElementBounds);
+			if (ElementBounds.SpheresIntersect(ElementBounds, WorldBrushSphereBounds))
 			{
-				ElementHandles.Emplace(ActorHandle);
+				ElementHandles.Add(MoveTemp(ResolvedHandle));
 			}
 		}
 	}
 
-	if (!SMInstanceElementDataUtil::SMInstanceElementsEnabled())
+	// Handle the IFA for the brush stroke level if needed for ISMs
+	const bool bAreFoliageSMInstancesEnabled = SMInstanceElementDataUtil::SMInstanceElementsEnabled() && FoliageElementUtil::FoliageInstanceElementsEnabled();
+	if (!bAreFoliageSMInstancesEnabled)
 	{
-		// Handle the IFA for the brush stroke level
 		if (UActorPartitionSubsystem* PartitionSubsystem = UWorld::GetSubsystem<UActorPartitionSubsystem>(GEditor->GetEditorWorldContext().World()))
 		{
 			constexpr bool bCreatePartitionActorIfMissing = false;
@@ -292,18 +325,21 @@ TArray<FTypedElementHandle> UPlacementBrushToolBase::GetElementsInBrushRadius()
 			{
 				for (const auto& FoliageInfo : FoliageActor->GetFoliageInfos())
 				{
+					// The foliage info needs a valid component for ISM placement from our systems.
+					if (FoliageInfo.Value->GetComponent() != nullptr)
+					{
+						continue;
+					}
+
 					FTypedElementHandle SourceObjectHandle = UEngineElementsLibrary::AcquireEditorObjectElementHandle(FoliageInfo.Key->GetSource());
-					if (GEditor->GetEditorSubsystem<UPlacementModeSubsystem>()->DoesCurrentPaletteSupportElement(SourceObjectHandle))
+					TArray<int32> Instances;
+					FSphere SphereToCheck(LastBrushStamp.WorldPosition, LastBrushStamp.Radius);
+					FoliageInfo.Value->GetInstancesInsideSphere(SphereToCheck, Instances);
+					if (Instances.Num())
 					{
-						TArray<int32> Instances;
-						FSphere SphereToCheck(LastBrushStamp.WorldPosition, LastBrushStamp.Radius);
-						FoliageInfo.Value->GetInstancesInsideSphere(SphereToCheck, Instances);
-						if (Instances.Num())
-						{
-							// For now, return the whole foliage actor, and allow the calling code to drill down, since we do not have element handles at the instance level just yet
-							ElementHandles.Emplace(UEngineElementsLibrary::AcquireEditorActorElementHandle(FoliageActor));
-							break;
-						}
+						// For now, return the whole foliage actor, and allow the calling code to drill down, since we do not have element handles at the instance level just yet
+						ElementHandles.Emplace(UEngineElementsLibrary::AcquireEditorActorElementHandle(FoliageActor));
+						break;
 					}
 				}
 			}
diff --git a/Engine/Source/Editor/AssetPlacementEdMode/Private/Tools/PlacementBrushToolBase.h b/Engine/Source/Editor/AssetPlacementEdMode/Private/Tools/PlacementBrushToolBase.h
--- a/Engine/Source/Editor/AssetPlacementEdMode/Private/Tools/PlacementBrushToolBase.h
+++ b/Engine/Source/Editor/AssetPlacementEdMode/Private/Tools/PlacementBrushToolBase.h
@@ -2,6 +2,7 @@
 
 #pragma once
 
+#include "InputState.h"
 #include "BaseTools/BaseBrushTool.h"
 #include "Elements/Framework/TypedElementHandle.h"
 
@@ -34,6 +35,9 @@ class UPlacementBrushToolBase : public UBaseBrushTool
 public:
 	virtual bool HitTest(const FRay& Ray, FHitResult& OutHit) override;
 	virtual bool AreAllTargetsValid() const override;
+	virtual void Render(IToolsContextRenderAPI* RenderAPI) override;
+	virtual void OnClickPress(const FInputDeviceRay& PressPos) override;
+	virtual void OnClickDrag(const FInputDeviceRay& DragPos) override;
 
 protected:
 	virtual double EstimateMaximumTargetDimension() override;
@@ -49,5 +53,8 @@ protected:
 	// Updates the last generated rotation to realign with the current brush position and normal.
 	static FQuat AlignRotationWithNormal(const FQuat& InRotation, const FVector& InNormal, EAxis::Type InAlignmentAxis, bool bInvertAxis);
 	static FTransform FinalizeTransform(const FTransform& OriginalTransform, const FVector& InNormal, const UAssetPlacementSettings* PlacementSettings);
-	TArray<FTypedElementHandle> GetElementsInBrushRadius() const;
+	TArray<FTypedElementHandle> GetElementsInBrushRadius(const FInputDeviceRay& DragPos) const;
+
+	float LastBrushStampWorldToPixelScale;
+	FInputDeviceRay LastDeviceInputRay = FInputDeviceRay(FRay());
 };
diff --git a/Engine/Source/Editor/AssetPlacementEdMode/Private/Tools/PlacementEraseTool.cpp b/Engine/Source/Editor/AssetPlacementEdMode/Private/Tools/PlacementEraseTool.cpp
--- a/Engine/Source/Editor/AssetPlacementEdMode/Private/Tools/PlacementEraseTool.cpp
+++ b/Engine/Source/Editor/AssetPlacementEdMode/Private/Tools/PlacementEraseTool.cpp
@@ -68,7 +68,7 @@ void UPlacementModeEraseTool::OnTick(float DeltaTime)
 
 	TGCObjectScopeGuard<UTypedElementList> ElementsToDelete(UTypedElementRegistry::GetInstance()->CreateElementList());
 
-	TArray<FTypedElementHandle> HitElements = GetElementsInBrushRadius();
+	TArray<FTypedElementHandle> HitElements = GetElementsInBrushRadius(LastDeviceInputRay);
 	for (const FTypedElementHandle& HitElement : HitElements)
 	{
 		if (TTypedElement<UTypedElementObjectInterface> ObjectInterface = UTypedElementRegistry::GetInstance()->GetElement<UTypedElementObjectInterface>(HitElement))
diff --git a/Engine/Source/Editor/AssetPlacementEdMode/Private/Tools/PlacementLassoSelectTool.cpp b/Engine/Source/Editor/AssetPlacementEdMode/Private/Tools/PlacementLassoSelectTool.cpp
--- a/Engine/Source/Editor/AssetPlacementEdMode/Private/Tools/PlacementLassoSelectTool.cpp
+++ b/Engine/Source/Editor/AssetPlacementEdMode/Private/Tools/PlacementLassoSelectTool.cpp
@@ -22,6 +22,11 @@
 
 constexpr TCHAR UPlacementModeLassoSelectTool::ToolName[];
 
+namespace PlacementModeLassoToolInternal
+{
+	FTypedElementSelectionOptions SelectionOptions {};
+}
+
 UPlacementBrushToolBase* UPlacementModeLassoSelectToolBuilder::FactoryToolInstance(UObject* Outer) const
 {
 	return NewObject<UPlacementModeLassoSelectTool>(Outer);
@@ -31,11 +36,55 @@ void UPlacementModeLassoSelectTool::OnBeginDrag(const FRay& Ray)
 {
 	Super::OnBeginDrag(Ray);
 
+	ElementsFromDrag.Empty();
 	GetToolManager()->BeginUndoTransaction(NSLOCTEXT("AssetPlacementEdMode", "BrushSelect", "Select Elements"));
 }
 
 void UPlacementModeLassoSelectTool::OnEndDrag(const FRay& Ray)
 {
+	if (IAssetEditorContextInterface* AssetEditorContext = GetToolManager()->GetContextObjectStore()->FindContext<IAssetEditorContextInterface>())
+	{
+		if (UTypedElementSelectionSet* SelectionSet = AssetEditorContext->GetMutableSelectionSet())
+		{
+			bool bSelectElements = !bCtrlToggle;
+			for (const FTypedElementHandle& HitElement : ElementsFromDrag)
+			{
+				if (!FoliageElementUtil::FoliageInstanceElementsEnabled())
+				{
+					if (TTypedElement<UTypedElementObjectInterface> ObjectInterface = UTypedElementRegistry::GetInstance()->GetElement<UTypedElementObjectInterface>(HitElement))
+					{
+						if (AInstancedFoliageActor* FoliageActor = ObjectInterface.GetObjectAs<AInstancedFoliageActor>())
+						{
+							FoliageActor->ForEachFoliageInfo([this, bSelectElements](UFoliageType* InFoliageType, FFoliageInfo& InFoliageInfo)
+								{
+									FTypedElementHandle SourceObjectHandle = UEngineElementsLibrary::AcquireEditorObjectElementHandle(InFoliageType->GetSource());
+									if (GEditor->GetEditorSubsystem<UPlacementModeSubsystem>()->DoesCurrentPaletteSupportElement(SourceObjectHandle))
+									{
+										TArray<int32> Instances;
+										FSphere SphereToCheck(LastBrushStamp.WorldPosition, LastBrushStamp.Radius);
+										InFoliageInfo.GetInstancesInsideSphere(SphereToCheck, Instances);
+										InFoliageInfo.SelectInstances(bSelectElements, Instances);
+									}
+									return true;
+								});
+
+							continue;
+						}
+					}
+				}
+
+				if (bSelectElements)
+				{
+					SelectionSet->SelectElement(HitElement, PlacementModeLassoToolInternal::SelectionOptions);
+				}
+				else
+				{
+					SelectionSet->DeselectElement(HitElement, PlacementModeLassoToolInternal::SelectionOptions);
+				}
+			}
+		}
+	}
+
 	GetToolManager()->EndUndoTransaction();
 
 	Super::OnEndDrag(Ray);
@@ -48,52 +97,5 @@ void UPlacementModeLassoSelectTool::OnTick(float DeltaTime)
 		return;
 	}
 
-	IAssetEditorContextInterface* AssetEditorContext = GetToolManager()->GetContextObjectStore()->FindContext<IAssetEditorContextInterface>();
-	if (!AssetEditorContext)
-	{
-		return;
-	}
-
-	UTypedElementSelectionSet* SelectionSet = AssetEditorContext->GetMutableSelectionSet();
-	if (!SelectionSet)
-	{
-		return;
-	}
-
-	bool bSelectElements = !bCtrlToggle;
-
-	TArray<FTypedElementHandle> HitElements = GetElementsInBrushRadius();
-	for (const FTypedElementHandle& HitElement : HitElements)
-	{
-		// Todo: Replace direct foliage instance selection with typed element selection
-		if (TTypedElement<UTypedElementObjectInterface> ObjectInterface = UTypedElementRegistry::GetInstance()->GetElement<UTypedElementObjectInterface>(HitElement))
-		{
-			if (AInstancedFoliageActor* FoliageActor = ObjectInterface.GetObjectAs<AInstancedFoliageActor>())
-			{
-				FoliageActor->ForEachFoliageInfo([this, bSelectElements](UFoliageType* InFoliageType, FFoliageInfo& InFoliageInfo)
-				{
-					FTypedElementHandle SourceObjectHandle = UEngineElementsLibrary::AcquireEditorObjectElementHandle(InFoliageType->GetSource());
-					if (GEditor->GetEditorSubsystem<UPlacementModeSubsystem>()->DoesCurrentPaletteSupportElement(SourceObjectHandle))
-					{
-						TArray<int32> Instances;
-						FSphere SphereToCheck(LastBrushStamp.WorldPosition, LastBrushStamp.Radius);
-						InFoliageInfo.GetInstancesInsideSphere(SphereToCheck, Instances);
-						InFoliageInfo.SelectInstances(bSelectElements, Instances);
-					}
-					return true;
-				});
-
-				continue;
-			}
-		}
-
-		if (bSelectElements)
-		{
-			SelectionSet->SelectElement(HitElement, FTypedElementSelectionOptions());
-		}
-		else
-		{
-			SelectionSet->DeselectElement(HitElement, FTypedElementSelectionOptions());
-		}
-	}
+	ElementsFromDrag.Append(GetElementsInBrushRadius(LastDeviceInputRay));
 }
diff --git a/Engine/Source/Editor/AssetPlacementEdMode/Private/Tools/PlacementLassoSelectTool.h b/Engine/Source/Editor/AssetPlacementEdMode/Private/Tools/PlacementLassoSelectTool.h
--- a/Engine/Source/Editor/AssetPlacementEdMode/Private/Tools/PlacementLassoSelectTool.h
+++ b/Engine/Source/Editor/AssetPlacementEdMode/Private/Tools/PlacementLassoSelectTool.h
@@ -27,4 +27,7 @@ public:
 	virtual void OnTick(float DeltaTime) override;
 
 	constexpr static TCHAR ToolName[] = TEXT("LassoSelectTool");
+
+private:
+	TSet<FTypedElementHandle> ElementsFromDrag;
 };
diff --git a/Engine/Source/Editor/AssetPlacementEdMode/Private/Tools/PlacementPlaceSingleTool.cpp b/Engine/Source/Editor/AssetPlacementEdMode/Private/Tools/PlacementPlaceSingleTool.cpp
--- a/Engine/Source/Editor/AssetPlacementEdMode/Private/Tools/PlacementPlaceSingleTool.cpp
+++ b/Engine/Source/Editor/AssetPlacementEdMode/Private/Tools/PlacementPlaceSingleTool.cpp
@@ -73,6 +73,8 @@ void UPlacementModePlaceSingleTool::OnClickPress(const FInputDeviceRay& PressPos
 	FPlacementOptions PlacementOptions;
 	PlacementOptions.bPreferBatchPlacement = true;
 	PlacementOptions.bPreferInstancedPlacement = SMInstanceElementDataUtil::SMInstanceElementsEnabled();
+
+	LastBrushStamp.Radius = 100.0f * LastBrushStampWorldToPixelScale;
 	
 	UPlacementSubsystem* PlacementSubsystem = GEditor->GetEditorSubsystem<UPlacementSubsystem>();
 	if (PlacementSubsystem)
@@ -224,15 +226,6 @@ FInputRayHit UPlacementModePlaceSingleTool::BeginHoverSequenceHitTest(const FInp
 	return FInputRayHit();
 }
 
-void UPlacementModePlaceSingleTool::Render(IToolsContextRenderAPI* RenderAPI)
-{
-	// Transform the brush radius to standard pixel size
-	float BrushRadiusScale = GizmoRenderingUtil::CalculateLocalPixelToWorldScale(RenderAPI->GetSceneView(), LastBrushStamp.WorldPosition);
-	LastBrushStamp.Radius = 100.0f * BrushRadiusScale;
-
-	Super::Render(RenderAPI);
-}
-
 void UPlacementModePlaceSingleTool::GeneratePlacementData(const FInputDeviceRay& DevicePos)
 {
 	const UAssetPlacementSettings* PlacementSettings = GEditor->GetEditorSubsystem<UPlacementModeSubsystem>()->GetModeSettingsObject();
@@ -295,6 +288,9 @@ void UPlacementModePlaceSingleTool::UpdatePreviewElements(const FInputDeviceRay&
 		return;
 	}
 
+	// Update the brush radius to be stable screen size based on the world to pixel scale at the last render update.
+	LastBrushStamp.Radius = 100.0f * LastBrushStampWorldToPixelScale;
+
 	FTransform UpdatedTransform(PlacementInfo->FinalizedTransform.GetRotation(), LastBrushStamp.WorldPosition, PlacementInfo->FinalizedTransform.GetScale3D());
 	UpdatedTransform = FinalizeTransform(UpdatedTransform, LastBrushStamp.WorldNormal, GEditor->GetEditorSubsystem<UPlacementModeSubsystem>()->GetModeSettingsObject());
 
diff --git a/Engine/Source/Editor/AssetPlacementEdMode/Private/Tools/PlacementPlaceSingleTool.h b/Engine/Source/Editor/AssetPlacementEdMode/Private/Tools/PlacementPlaceSingleTool.h
--- a/Engine/Source/Editor/AssetPlacementEdMode/Private/Tools/PlacementPlaceSingleTool.h
+++ b/Engine/Source/Editor/AssetPlacementEdMode/Private/Tools/PlacementPlaceSingleTool.h
@@ -65,8 +65,6 @@ public:
 	virtual FInputRayHit CanBeginClickDragSequence(const FInputDeviceRay& PressPos) override;
 	virtual FInputRayHit BeginHoverSequenceHitTest(const FInputDeviceRay& DevicePos) override;
 
-	virtual void Render(IToolsContextRenderAPI* RenderAPI) override;
-
 protected:
 	void GeneratePlacementData(const FInputDeviceRay& DevicePos);
 	void CreatePreviewElements(const FInputDeviceRay& DevicePos);
diff --git a/Engine/Source/Runtime/Foliage/Public/InstancedFoliage.h b/Engine/Source/Runtime/Foliage/Public/InstancedFoliage.h
--- a/Engine/Source/Runtime/Foliage/Public/InstancedFoliage.h
+++ b/Engine/Source/Runtime/Foliage/Public/InstancedFoliage.h
@@ -23,6 +23,11 @@ struct FFoliageInstanceHash;
 
 DECLARE_LOG_CATEGORY_EXTERN(LogInstancedFoliage, Log, All);
 
+namespace FoliageElementUtil
+{
+	FOLIAGE_API bool FoliageInstanceElementsEnabled();
+}
+
 /**
 * Flags stored with each instance
 */
