Cleaned up some code now that FTypedElementList has a well defined lifetime

 - Added FTypedElementListConstRef overloads of SelectElements, DeselectElements, and SetSelection in UTypedElementSelectionSet.
 - Added FTypedElementListConstRef overload of Append in FTypedElementList.
 - Removed the TSet<FTypedElementHandle> overload of GetElementHandlesInRect in FViewport.
 - Updated UPlacementBrushToolBase::GetElementsInBrushRadius to return/use a FTypedElementListRef, and updated UPlacementModeLassoSelectTool and UPlacementModeEraseTool accordingly.
 - Updated FoliageElementUtil::FoliageInstanceElementsEnabled to only return true if static mesh instance elements are also enabled.

diff --git a/Engine/Source/Editor/AssetPlacementEdMode/Private/Tools/PlacementBrushToolBase.cpp b/Engine/Source/Editor/AssetPlacementEdMode/Private/Tools/PlacementBrushToolBase.cpp
--- a/Engine/Source/Editor/AssetPlacementEdMode/Private/Tools/PlacementBrushToolBase.cpp
+++ b/Engine/Source/Editor/AssetPlacementEdMode/Private/Tools/PlacementBrushToolBase.cpp
@@ -22,7 +22,6 @@
 #include "Modes/PlacementModeSubsystem.h"
 #include "ActorFactories/ActorFactory.h"
 #include "BaseGizmos/GizmoRenderingUtil.h"
-#include "Elements/SMInstance/SMInstanceElementData.h" // For SMInstanceElementDataUtil::SMInstanceElementsEnabled
 
 bool UPlacementToolBuilderBase::CanBuildTool(const FToolBuilderState& SceneState) const
 {	
@@ -262,9 +261,9 @@ FTransform UPlacementBrushToolBase::FinalizeTransform(const FTransform& Original
 	return FinalizedTransform;
 }
 
-TArray<FTypedElementHandle> UPlacementBrushToolBase::GetElementsInBrushRadius(const FInputDeviceRay& DragPos) const
+FTypedElementListRef UPlacementBrushToolBase::GetElementsInBrushRadius(const FInputDeviceRay& DragPos) const
 {
-	TArray<FTypedElementHandle> ElementHandles;
+	FTypedElementListRef ElementHandles = UTypedElementRegistry::GetInstance()->CreateElementList();
 
 	// We need the 2D device screen space position to test against hit proxies.
 	if (!DragPos.bHas2D)
@@ -295,27 +294,29 @@ TArray<FTypedElementHandle> UPlacementBrushToolBase::GetElementsInBrushRadius(co
 	AreaToCheck.Max.X = FMath::Min<int32>(ViewportSize.X, DragPos.ScreenPosition.X + HalfBrushRadius);
 	AreaToCheck.Max.Y = FMath::Min<int32>(ViewportSize.Y, DragPos.ScreenPosition.Y + HalfBrushRadius);
 
-	FBoxSphereBounds WorldBrushSphereBounds(FSphere(LastBrushStamp.WorldPosition, LastBrushStamp.Radius));
-	TSet<FTypedElementHandle> HitElementHandles;
+	// Get the raw hit proxies within the rect
+	FTypedElementListRef HitElementHandles = UTypedElementRegistry::GetInstance()->CreateElementList();
 	Viewport->GetElementHandlesInRect(AreaToCheck, HitElementHandles);
-	for (const FTypedElementHandle& HitHandle : HitElementHandles)
+
+	// Work out which elements to actually select in the viewport, and also verify that we're actually intersecting with the sphere from the brush in world space.
+	FBoxSphereBounds WorldBrushSphereBounds(FSphere(LastBrushStamp.WorldPosition, LastBrushStamp.Radius));
+	HitElementHandles->ForEachElementHandle([SelectionSet, ElementHandles, &WorldBrushSphereBounds](const FTypedElementHandle& HitHandle)
 	{
 		FTypedElementHandle ResolvedHandle = SelectionSet->GetSelectionElement(HitHandle, ETypedElementSelectionMethod::Primary);
-		if (TTypedElement<UTypedElementWorldInterface> WorldInterface = UTypedElementRegistry::GetInstance()->GetElement<UTypedElementWorldInterface>(ResolvedHandle))
+		if (TTypedElement<UTypedElementWorldInterface> WorldInterface = SelectionSet->GetElementList()->GetElement<UTypedElementWorldInterface>(ResolvedHandle))
 		{
-			// Since the viewport gives us back hit proxies in a rect, we also want to verify that we're actually intersecting with the sphere from the brush in world space.
 			FBoxSphereBounds ElementBounds;
 			WorldInterface.GetBounds(ElementBounds);
 			if (ElementBounds.SpheresIntersect(ElementBounds, WorldBrushSphereBounds))
 			{
-				ElementHandles.Add(MoveTemp(ResolvedHandle));
+				ElementHandles->Add(MoveTemp(ResolvedHandle));
 			}
 		}
-	}
+		return true;
+	});
 
 	// Handle the IFA for the brush stroke level if needed for ISMs
-	const bool bAreFoliageSMInstancesEnabled = SMInstanceElementDataUtil::SMInstanceElementsEnabled() && FoliageElementUtil::FoliageInstanceElementsEnabled();
-	if (!bAreFoliageSMInstancesEnabled)
+	if (!FoliageElementUtil::FoliageInstanceElementsEnabled())
 	{
 		if (UActorPartitionSubsystem* PartitionSubsystem = UWorld::GetSubsystem<UActorPartitionSubsystem>(GEditor->GetEditorWorldContext().World()))
 		{
@@ -338,7 +339,7 @@ TArray<FTypedElementHandle> UPlacementBrushToolBase::GetElementsInBrushRadius(co
 					if (Instances.Num())
 					{
 						// For now, return the whole foliage actor, and allow the calling code to drill down, since we do not have element handles at the instance level just yet
-						ElementHandles.Emplace(UEngineElementsLibrary::AcquireEditorActorElementHandle(FoliageActor));
+						ElementHandles->Add(UEngineElementsLibrary::AcquireEditorActorElementHandle(FoliageActor));
 						break;
 					}
 				}
diff --git a/Engine/Source/Editor/AssetPlacementEdMode/Private/Tools/PlacementBrushToolBase.h b/Engine/Source/Editor/AssetPlacementEdMode/Private/Tools/PlacementBrushToolBase.h
--- a/Engine/Source/Editor/AssetPlacementEdMode/Private/Tools/PlacementBrushToolBase.h
+++ b/Engine/Source/Editor/AssetPlacementEdMode/Private/Tools/PlacementBrushToolBase.h
@@ -4,7 +4,7 @@
 
 #include "InputState.h"
 #include "BaseTools/BaseBrushTool.h"
-#include "Elements/Framework/TypedElementHandle.h"
+#include "Elements/Framework/TypedElementListFwd.h"
 
 #include "PlacementBrushToolBase.generated.h"
 
@@ -53,7 +53,7 @@ protected:
 	// Updates the last generated rotation to realign with the current brush position and normal.
 	static FQuat AlignRotationWithNormal(const FQuat& InRotation, const FVector& InNormal, EAxis::Type InAlignmentAxis, bool bInvertAxis);
 	static FTransform FinalizeTransform(const FTransform& OriginalTransform, const FVector& InNormal, const UAssetPlacementSettings* PlacementSettings);
-	TArray<FTypedElementHandle> GetElementsInBrushRadius(const FInputDeviceRay& DragPos) const;
+	FTypedElementListRef GetElementsInBrushRadius(const FInputDeviceRay& DragPos) const;
 
 	float LastBrushStampWorldToPixelScale;
 	FInputDeviceRay LastDeviceInputRay = FInputDeviceRay(FRay());
diff --git a/Engine/Source/Editor/AssetPlacementEdMode/Private/Tools/PlacementEraseTool.cpp b/Engine/Source/Editor/AssetPlacementEdMode/Private/Tools/PlacementEraseTool.cpp
--- a/Engine/Source/Editor/AssetPlacementEdMode/Private/Tools/PlacementEraseTool.cpp
+++ b/Engine/Source/Editor/AssetPlacementEdMode/Private/Tools/PlacementEraseTool.cpp
@@ -66,12 +66,10 @@ void UPlacementModeEraseTool::OnTick(float DeltaTime)
 		return;
 	}
 
-	FTypedElementListRef ElementsToDelete = UTypedElementRegistry::GetInstance()->CreateElementList();
-
-	TArray<FTypedElementHandle> HitElements = GetElementsInBrushRadius(LastDeviceInputRay);
-	for (const FTypedElementHandle& HitElement : HitElements)
+	FTypedElementListRef ElementsToDelete = GetElementsInBrushRadius(LastDeviceInputRay);
+	if (!FoliageElementUtil::FoliageInstanceElementsEnabled())
 	{
-		if (TTypedElement<UTypedElementObjectInterface> ObjectInterface = UTypedElementRegistry::GetInstance()->GetElement<UTypedElementObjectInterface>(HitElement))
+		ElementsToDelete->RemoveAll<UTypedElementObjectInterface>([this](const TTypedElement<UTypedElementObjectInterface>& ObjectInterface)
 		{
 			// Since the foliage static mesh instances do not currently operate with element handles, we have to drill in manually here.
 			if (AInstancedFoliageActor* FoliageActor = ObjectInterface.GetObjectAs<AInstancedFoliageActor>())
@@ -88,12 +86,10 @@ void UPlacementModeEraseTool::OnTick(float DeltaTime)
 					}
 					return true; // continue iteration
 				});
+				return true; // Foliage - remove from the normal element delete
 			}
-			else
-			{
-				ElementsToDelete->Add(HitElement);
-			}
-		}
+			return false; // Not foliage - will be processed via the normal element delete
+		});
 	}
 
 	if (ElementsToDelete->HasElements())
diff --git a/Engine/Source/Editor/AssetPlacementEdMode/Private/Tools/PlacementLassoSelectTool.cpp b/Engine/Source/Editor/AssetPlacementEdMode/Private/Tools/PlacementLassoSelectTool.cpp
--- a/Engine/Source/Editor/AssetPlacementEdMode/Private/Tools/PlacementLassoSelectTool.cpp
+++ b/Engine/Source/Editor/AssetPlacementEdMode/Private/Tools/PlacementLassoSelectTool.cpp
@@ -36,7 +36,7 @@ void UPlacementModeLassoSelectTool::OnBeginDrag(const FRay& Ray)
 {
 	Super::OnBeginDrag(Ray);
 
-	ElementsFromDrag.Empty();
+	ElementsFromDrag.Reset();
 	GetToolManager()->BeginUndoTransaction(NSLOCTEXT("AssetPlacementEdMode", "BrushSelect", "Select Elements"));
 }
 
@@ -44,48 +44,49 @@ void UPlacementModeLassoSelectTool::OnEndDrag(const FRay& Ray)
 {
 	if (IAssetEditorContextInterface* AssetEditorContext = GetToolManager()->GetContextObjectStore()->FindContext<IAssetEditorContextInterface>())
 	{
-		if (UTypedElementSelectionSet* SelectionSet = AssetEditorContext->GetMutableSelectionSet())
+		UTypedElementSelectionSet* SelectionSet = AssetEditorContext->GetMutableSelectionSet();
+		if (SelectionSet && ElementsFromDrag)
 		{
 			bool bSelectElements = !bCtrlToggle;
-			for (const FTypedElementHandle& HitElement : ElementsFromDrag)
+
+			if (!FoliageElementUtil::FoliageInstanceElementsEnabled())
 			{
-				if (!FoliageElementUtil::FoliageInstanceElementsEnabled())
+				ElementsFromDrag->RemoveAll<UTypedElementObjectInterface>([this, bSelectElements](const TTypedElement<UTypedElementObjectInterface>& ObjectInterface)
 				{
-					if (TTypedElement<UTypedElementObjectInterface> ObjectInterface = UTypedElementRegistry::GetInstance()->GetElement<UTypedElementObjectInterface>(HitElement))
+					// Since the foliage static mesh instances do not currently operate with element handles, we have to drill in manually here.
+					if (AInstancedFoliageActor* FoliageActor = ObjectInterface.GetObjectAs<AInstancedFoliageActor>())
 					{
-						if (AInstancedFoliageActor* FoliageActor = ObjectInterface.GetObjectAs<AInstancedFoliageActor>())
+						FoliageActor->ForEachFoliageInfo([this, bSelectElements](UFoliageType* FoliageType, FFoliageInfo& FoliageInfo)
 						{
-							FoliageActor->ForEachFoliageInfo([this, bSelectElements](UFoliageType* InFoliageType, FFoliageInfo& InFoliageInfo)
-								{
-									FTypedElementHandle SourceObjectHandle = UEngineElementsLibrary::AcquireEditorObjectElementHandle(InFoliageType->GetSource());
-									if (GEditor->GetEditorSubsystem<UPlacementModeSubsystem>()->DoesCurrentPaletteSupportElement(SourceObjectHandle))
-									{
-										TArray<int32> Instances;
-										FSphere SphereToCheck(LastBrushStamp.WorldPosition, LastBrushStamp.Radius);
-										InFoliageInfo.GetInstancesInsideSphere(SphereToCheck, Instances);
-										InFoliageInfo.SelectInstances(bSelectElements, Instances);
-									}
-									return true;
-								});
-
-							continue;
-						}
+							FTypedElementHandle SourceObjectHandle = UEngineElementsLibrary::AcquireEditorObjectElementHandle(FoliageType->GetSource());
+							if (GEditor->GetEditorSubsystem<UPlacementModeSubsystem>()->DoesCurrentPaletteSupportElement(SourceObjectHandle))
+							{
+								TArray<int32> Instances;
+								FSphere SphereToCheck(LastBrushStamp.WorldPosition, LastBrushStamp.Radius);
+								FoliageInfo.GetInstancesInsideSphere(SphereToCheck, Instances);
+								FoliageInfo.SelectInstances(bSelectElements, Instances);
+							}
+							return true; // continue iteration
+						});
+						return true; // Foliage - remove from the normal element select
 					}
-				}
+					return false; // Not foliage - will be processed via the normal element select
+				});
+			}
 
-				if (bSelectElements)
-				{
-					SelectionSet->SelectElement(HitElement, PlacementModeLassoToolInternal::SelectionOptions);
-				}
-				else
-				{
-					SelectionSet->DeselectElement(HitElement, PlacementModeLassoToolInternal::SelectionOptions);
-				}
+			if (bSelectElements)
+			{
+				SelectionSet->SelectElements(ElementsFromDrag.ToSharedRef(), PlacementModeLassoToolInternal::SelectionOptions);
+			}
+			else
+			{
+				SelectionSet->DeselectElements(ElementsFromDrag.ToSharedRef(), PlacementModeLassoToolInternal::SelectionOptions);
 			}
 		}
 	}
 
 	GetToolManager()->EndUndoTransaction();
+	ElementsFromDrag.Reset();
 
 	Super::OnEndDrag(Ray);
 }
@@ -97,5 +98,13 @@ void UPlacementModeLassoSelectTool::OnTick(float DeltaTime)
 		return;
 	}
 
-	ElementsFromDrag.Append(GetElementsInBrushRadius(LastDeviceInputRay));
+	FTypedElementListRef HitElements = GetElementsInBrushRadius(LastDeviceInputRay);
+	if (ElementsFromDrag)
+	{
+		ElementsFromDrag->Append(HitElements);
+	}
+	else
+	{
+		ElementsFromDrag = HitElements;
+	}
 }
diff --git a/Engine/Source/Editor/AssetPlacementEdMode/Private/Tools/PlacementLassoSelectTool.h b/Engine/Source/Editor/AssetPlacementEdMode/Private/Tools/PlacementLassoSelectTool.h
--- a/Engine/Source/Editor/AssetPlacementEdMode/Private/Tools/PlacementLassoSelectTool.h
+++ b/Engine/Source/Editor/AssetPlacementEdMode/Private/Tools/PlacementLassoSelectTool.h
@@ -29,5 +29,5 @@ public:
 	constexpr static TCHAR ToolName[] = TEXT("LassoSelectTool");
 
 private:
-	TSet<FTypedElementHandle> ElementsFromDrag;
+	FTypedElementListPtr ElementsFromDrag;
 };
diff --git a/Engine/Source/Runtime/Engine/Private/UnrealClient.cpp b/Engine/Source/Runtime/Engine/Private/UnrealClient.cpp
--- a/Engine/Source/Runtime/Engine/Private/UnrealClient.cpp
+++ b/Engine/Source/Runtime/Engine/Private/UnrealClient.cpp
@@ -1866,24 +1866,13 @@ FTypedElementHandle FViewport::GetElementHandleAtPoint(int32 X, int32 Y)
 
 void FViewport::GetElementHandlesInRect(FIntRect InRect, FTypedElementListRef OutElementHandles)
 {
-	GetElementHandlesInRectImpl(InRect, *OutElementHandles);
-}
-
-void FViewport::GetElementHandlesInRect(FIntRect InRect, TSet<FTypedElementHandle>& OutElementHandles)
-{
-	GetElementHandlesInRectImpl(InRect, OutElementHandles);
-}
-
-template <typename ContainerType>
-void FViewport::GetElementHandlesInRectImpl(FIntRect InRect, ContainerType& OutElementHandles)
-{
-	OutElementHandles.Reset();
+	OutElementHandles->Reset();
 
 	EnumerateHitProxiesInRect(InRect, [&OutElementHandles](HHitProxy* HitProxy)
 	{
 		if (FTypedElementHandle ElementHandle = HitProxy->GetElementHandle())
 		{
-			OutElementHandles.Add(MoveTemp(ElementHandle));
+			OutElementHandles->Add(MoveTemp(ElementHandle));
 		}
 		return true;
 	});
diff --git a/Engine/Source/Runtime/Engine/Public/UnrealClient.h b/Engine/Source/Runtime/Engine/Public/UnrealClient.h
--- a/Engine/Source/Runtime/Engine/Public/UnrealClient.h
+++ b/Engine/Source/Runtime/Engine/Public/UnrealClient.h
@@ -509,7 +509,6 @@ public:
 	 * If the hit proxies are not cached, this will call ViewportClient->Draw with a hit-testing canvas.
 	 */
 	ENGINE_API void GetElementHandlesInRect(FIntRect InRect, FTypedElementListRef OutElementHandles);
-	ENGINE_API void GetElementHandlesInRect(FIntRect InRect, TSet<FTypedElementHandle>& OutElementHandles);
 
 	/**
 	 * Retrieves the interface to the viewport's frame, if it has one.
@@ -646,14 +645,6 @@ protected:
 
 private:
 
-	/**
-	 * Returns all element handles found within a specified region.
-	 * InRect must be entirely within the viewport's client area.
-	 * If the hit proxies are not cached, this will call ViewportClient->Draw with a hit-testing canvas.
-	 */
-	template <typename ContainerType>
-	void GetElementHandlesInRectImpl(FIntRect InRect, ContainerType& OutElementHandles);
-
 	/**
 	 * Enumerate all valid hit proxies found within a specified region.
 	 * InRect must be entirely within the viewport's client area.
diff --git a/Engine/Source/Runtime/Foliage/Private/InstancedFoliage.cpp b/Engine/Source/Runtime/Foliage/Private/InstancedFoliage.cpp
--- a/Engine/Source/Runtime/Foliage/Private/InstancedFoliage.cpp
+++ b/Engine/Source/Runtime/Foliage/Private/InstancedFoliage.cpp
@@ -50,6 +50,7 @@ InstancedFoliage.cpp: Instanced foliage implementation.
 #include "LevelUtils.h"
 #include "FoliageHelper.h"
 #include "Elements/Framework/EngineElementsLibrary.h"
+#include "Elements/SMInstance/SMInstanceElementData.h" // For SMInstanceElementDataUtil::SMInstanceElementsEnabled
 #include "Algo/Transform.h"
 #include "ActorPartition/ActorPartitionSubsystem.h"
 #include "Misc/CoreMisc.h"
@@ -88,7 +89,7 @@ static FAutoConsoleVariableRef CVarEnableFoliageInstanceElements(
 
 bool FoliageInstanceElementsEnabled()
 {
-	return GEnableFoliageInstanceElements != 0;
+	return SMInstanceElementDataUtil::SMInstanceElementsEnabled() && GEnableFoliageInstanceElements != 0;
 }
 
 }
diff --git a/Engine/Source/Runtime/TypedElementFramework/Private/Elements/Framework/TypedElementListLibrary.h b/Engine/Source/Runtime/TypedElementFramework/Private/Elements/Framework/TypedElementListLibrary.h
--- a/Engine/Source/Runtime/TypedElementFramework/Private/Elements/Framework/TypedElementListLibrary.h
+++ b/Engine/Source/Runtime/TypedElementFramework/Private/Elements/Framework/TypedElementListLibrary.h
@@ -195,7 +195,7 @@ public:
 	}
 
 	/**
-	 * Append the given element handles to this element list, for any that already in the list.
+	 * Append the given element handles to this element list.
 	 */
 	UFUNCTION(BlueprintCallable, Category="TypedElementFramework|List", meta=(ScriptMethod))
 	static void Append(FTypedElementListProxy ElementList, const TArray<FTypedElementHandle>& ElementHandles)
@@ -206,6 +206,20 @@ public:
 		}
 	}
 
+	/**
+	 * Append the another element list to this element list.
+	 */
+	UFUNCTION(BlueprintCallable, Category="TypedElementFramework|List", meta=(ScriptMethod))
+	static void AppendList(FTypedElementListProxy ElementList, const FTypedElementListProxy OtherElementList)
+	{
+		FTypedElementListPtr ElementListPtr = ElementList.GetElementList();
+		FTypedElementListConstPtr OtherElementListPtr = OtherElementList.GetElementList();
+		if (ElementListPtr && OtherElementListPtr)
+		{
+			ElementListPtr->Append(OtherElementListPtr.ToSharedRef());
+		}
+	}
+
 	/**
 	 * Remove the given element handle from this element list, if it is in the list.
 	 * @return True if the element handle was removed, false if it isn't in the list.
diff --git a/Engine/Source/Runtime/TypedElementFramework/Public/Elements/Framework/TypedElementList.h b/Engine/Source/Runtime/TypedElementFramework/Public/Elements/Framework/TypedElementList.h
--- a/Engine/Source/Runtime/TypedElementFramework/Public/Elements/Framework/TypedElementList.h
+++ b/Engine/Source/Runtime/TypedElementFramework/Public/Elements/Framework/TypedElementList.h
@@ -428,7 +428,25 @@ public:
 	}
 
 	/**
-	 * Append the given element handles to this element list, for any that already in the list.
+	 * Append another element list to this element list.
+	 */
+	void Append(FTypedElementListConstRef InElementList)
+	{
+		if (this != &InElementList.Get())
+		{
+			FTypedElementListLegacySyncScopedBatch LegacySyncBatch(*this);
+
+			Reserve(Num() + InElementList->Num());
+			InElementList->ForEachElementHandle([this](const FTypedElementHandle& ElementHandle)
+			{
+				AddElementImpl(CopyTemp(ElementHandle));
+				return true;
+			});
+		}
+	}
+
+	/**
+	 * Append the given element handles to this element list.
 	 */
 	void Append(TArrayView<const FTypedElementHandle> InElementHandles)
 	{
@@ -442,7 +460,7 @@ public:
 	}
 
 	/**
-	 * Append the given element owners to this element list, for any that already in the list.
+	 * Append the given element owners to this element list.
 	 */
 	template <typename ElementDataType>
 	FORCEINLINE void Append(const TArray<TTypedElementOwner<ElementDataType>>& InElementOwners)
@@ -451,7 +469,7 @@ public:
 	}
 
 	/**
-	 * Append the given element owners to this element list, for any that already in the list.
+	 * Append the given element owners to this element list.
 	 */
 	template <typename ElementDataType>
 	void Append(TArrayView<const TTypedElementOwner<ElementDataType>> InElementOwners)
diff --git a/Engine/Source/Runtime/TypedElementRuntime/Private/Elements/Framework/TypedElementSelectionSet.cpp b/Engine/Source/Runtime/TypedElementRuntime/Private/Elements/Framework/TypedElementSelectionSet.cpp
--- a/Engine/Source/Runtime/TypedElementRuntime/Private/Elements/Framework/TypedElementSelectionSet.cpp
+++ b/Engine/Source/Runtime/TypedElementRuntime/Private/Elements/Framework/TypedElementSelectionSet.cpp
@@ -213,6 +213,21 @@ bool UTypedElementSelectionSet::SelectElements(TArrayView<const FTypedElementHan
 	return bSelectionChanged;
 }
 
+bool UTypedElementSelectionSet::SelectElements(FTypedElementListConstRef InElementList, const FTypedElementSelectionOptions InSelectionOptions)
+{
+	FTypedElementListLegacySyncScopedBatch LegacySyncBatch(*ElementList, InSelectionOptions.AllowLegacyNotifications());
+
+	bool bSelectionChanged = false;
+
+	InElementList->ForEachElementHandle([this, &bSelectionChanged, &InSelectionOptions](const FTypedElementHandle& ElementHandle)
+	{
+		bSelectionChanged |= SelectElement(ElementHandle, InSelectionOptions);
+		return true;
+	});
+
+	return bSelectionChanged;
+}
+
 bool UTypedElementSelectionSet::DeselectElement(const FTypedElementHandle& InElementHandle, const FTypedElementSelectionOptions InSelectionOptions)
 {
 	bool bSelectionChanged = false;
@@ -263,6 +278,21 @@ bool UTypedElementSelectionSet::DeselectElements(TArrayView<const FTypedElementH
 	return bSelectionChanged;
 }
 
+bool UTypedElementSelectionSet::DeselectElements(FTypedElementListConstRef InElementList, const FTypedElementSelectionOptions InSelectionOptions)
+{
+	FTypedElementListLegacySyncScopedBatch LegacySyncBatch(*ElementList, InSelectionOptions.AllowLegacyNotifications());
+
+	bool bSelectionChanged = false;
+
+	InElementList->ForEachElementHandle([this, &bSelectionChanged, &InSelectionOptions](const FTypedElementHandle& ElementHandle)
+	{
+		bSelectionChanged |= DeselectElement(ElementHandle, InSelectionOptions);
+		return true;
+	});
+
+	return bSelectionChanged;
+}
+
 bool UTypedElementSelectionSet::ClearSelection(const FTypedElementSelectionOptions InSelectionOptions)
 {
 	FTypedElementListLegacySyncScopedBatch LegacySyncBatch(*ElementList, InSelectionOptions.AllowLegacyNotifications());
@@ -310,6 +340,18 @@ bool UTypedElementSelectionSet::SetSelection(TArrayView<const FTypedElementHandl
 	return bSelectionChanged;
 }
 
+bool UTypedElementSelectionSet::SetSelection(FTypedElementListConstRef InElementList, const FTypedElementSelectionOptions InSelectionOptions)
+{
+	FTypedElementListLegacySyncScopedBatch LegacySyncBatch(*ElementList, InSelectionOptions.AllowLegacyNotifications());
+
+	bool bSelectionChanged = false;
+
+	bSelectionChanged |= ClearSelection(InSelectionOptions);
+	bSelectionChanged |= SelectElements(InElementList, InSelectionOptions);
+
+	return bSelectionChanged;
+}
+
 bool UTypedElementSelectionSet::AllowSelectionModifiers(const FTypedElementHandle& InElementHandle) const
 {
 	FTypedElementSelectionSetElement SelectionSetElement = ResolveSelectionSetElement(InElementHandle);
diff --git a/Engine/Source/Runtime/TypedElementRuntime/Private/Elements/Framework/TypedElementSelectionSetLibrary.h b/Engine/Source/Runtime/TypedElementRuntime/Private/Elements/Framework/TypedElementSelectionSetLibrary.h
--- a/Engine/Source/Runtime/TypedElementRuntime/Private/Elements/Framework/TypedElementSelectionSetLibrary.h
+++ b/Engine/Source/Runtime/TypedElementRuntime/Private/Elements/Framework/TypedElementSelectionSetLibrary.h
@@ -13,6 +13,42 @@ class UTypedElementSelectionSetLibrary : public UObject
 	GENERATED_BODY()
 	
 public:
+	/**
+	 * Attempt to select the given elements.
+	 * @return True if the selection was changed, false otherwise.
+	 */
+	UFUNCTION(BlueprintCallable, Category="TypedElementFramework|Selection", meta=(ScriptMethod))
+	static bool SelectElementsFromList(UTypedElementSelectionSet* SelectionSet, const FTypedElementListProxy ElementList, const FTypedElementSelectionOptions SelectionOptions)
+	{
+		FTypedElementListConstPtr ElementListPtr = ElementList.GetElementList();
+		return ElementListPtr && SelectionSet->SelectElements(ElementListPtr.ToSharedRef(), SelectionOptions);
+	}
+
+	/**
+	 * Attempt to deselect the given elements.
+	 * @return True if the selection was changed, false otherwise.
+	 */
+	UFUNCTION(BlueprintCallable, Category="TypedElementFramework|Selection", meta=(ScriptMethod))
+	static bool DeselectElementsFromList(UTypedElementSelectionSet* SelectionSet, const FTypedElementListProxy ElementList, const FTypedElementSelectionOptions SelectionOptions)
+	{
+		FTypedElementListConstPtr ElementListPtr = ElementList.GetElementList();
+		return ElementListPtr && SelectionSet->DeselectElements(ElementListPtr.ToSharedRef(), SelectionOptions);
+	}
+
+	/**
+	 * Attempt to make the selection the given elements.
+	 * @note Equivalent to calling ClearSelection then SelectElements, but happens in a single batch.
+	 * @return True if the selection was changed, false otherwise.
+	 */
+	UFUNCTION(BlueprintCallable, Category="TypedElementFramework|Selection", meta=(ScriptMethod))
+	static bool SetSelectionFromList(UTypedElementSelectionSet* SelectionSet, const FTypedElementListProxy ElementList, const FTypedElementSelectionOptions SelectionOptions)
+	{
+		FTypedElementListConstPtr ElementListPtr = ElementList.GetElementList();
+		return ElementListPtr
+			? SelectionSet->SetSelection(ElementListPtr.ToSharedRef(), SelectionOptions)
+			: SelectionSet->ClearSelection(SelectionOptions);
+	}
+
 	/**
 	 * Get a normalized version of this selection set that can be used to perform operations like gizmo manipulation, deletion, copying, etc.
 	 * This will do things like expand out groups, and resolve any parent<->child elements so that duplication operations aren't performed on both the parent and the child.
diff --git a/Engine/Source/Runtime/TypedElementRuntime/Public/Elements/Framework/TypedElementSelectionSet.h b/Engine/Source/Runtime/TypedElementRuntime/Public/Elements/Framework/TypedElementSelectionSet.h
--- a/Engine/Source/Runtime/TypedElementRuntime/Public/Elements/Framework/TypedElementSelectionSet.h
+++ b/Engine/Source/Runtime/TypedElementRuntime/Public/Elements/Framework/TypedElementSelectionSet.h
@@ -193,6 +193,7 @@ public:
 	UFUNCTION(BlueprintCallable, Category="TypedElementFramework|Selection")
 	bool SelectElements(const TArray<FTypedElementHandle>& InElementHandles, const FTypedElementSelectionOptions InSelectionOptions);
 	bool SelectElements(TArrayView<const FTypedElementHandle> InElementHandles, const FTypedElementSelectionOptions InSelectionOptions);
+	bool SelectElements(FTypedElementListConstRef InElementList, const FTypedElementSelectionOptions InSelectionOptions);
 
 	/**
 	 * Attempt to deselect the given element.
@@ -208,6 +209,7 @@ public:
 	UFUNCTION(BlueprintCallable, Category="TypedElementFramework|Selection")
 	bool DeselectElements(const TArray<FTypedElementHandle>& InElementHandles, const FTypedElementSelectionOptions InSelectionOptions);
 	bool DeselectElements(TArrayView<const FTypedElementHandle> InElementHandles, const FTypedElementSelectionOptions InSelectionOptions);
+	bool DeselectElements(FTypedElementListConstRef InElementList, const FTypedElementSelectionOptions InSelectionOptions);
 
 	/**
 	 * Clear the current selection.
@@ -224,6 +226,7 @@ public:
 	UFUNCTION(BlueprintCallable, Category="TypedElementFramework|Selection")
 	bool SetSelection(const TArray<FTypedElementHandle>& InElementHandles, const FTypedElementSelectionOptions InSelectionOptions);
 	bool SetSelection(TArrayView<const FTypedElementHandle> InElementHandles, const FTypedElementSelectionOptions InSelectionOptions);
+	bool SetSelection(FTypedElementListConstRef InElementList, const FTypedElementSelectionOptions InSelectionOptions);
 
 	/**
 	 * Test to see whether selection modifiers (Ctrl or Shift) are allowed while selecting this element.
