Asset Placement Editor Mode now has a separate tab to hold the asset palette.

diff --git a/Engine/Source/Editor/AssetPlacementEdMode/Private/AssetPlacementEdModeToolkit.cpp b/Engine/Source/Editor/AssetPlacementEdMode/Private/AssetPlacementEdModeToolkit.cpp
--- a/Engine/Source/Editor/AssetPlacementEdMode/Private/AssetPlacementEdModeToolkit.cpp
+++ b/Engine/Source/Editor/AssetPlacementEdMode/Private/AssetPlacementEdModeToolkit.cpp
@@ -6,6 +6,7 @@
 #include "SAssetPlacementPalette.h"
 #include "Widgets/Layout/SExpandableArea.h"
 #include "Widgets/Layout/SScrollBox.h"
+#include "Toolkits/AssetEditorModeUILayer.h"
 
 #define LOCTEXT_NAMESPACE "AssetPlacementEdModeToolkit"
 
@@ -29,36 +30,64 @@ FText FAssetPlacementEdModeToolkit::GetBaseToolkitName() const
 	return LOCTEXT("DisplayName", "AssetPlacementEdMode Tool");
 }
 
+void FAssetPlacementEdModeToolkit::InvokeUI()
+{
+	FModeToolkit::InvokeUI();
+
+	if (ModeUILayer.IsValid())
+	{
+		TSharedPtr<FAssetEditorModeUILayer> ModeUILayerPtr = ModeUILayer.Pin();
+		AssetPaletteTab = ModeUILayerPtr->GetTabManager()->TryInvokeTab(FAssetEditorModeUILayer::BottomLeftTabID);
+	}
+}
+
 TSharedPtr<SWidget> FAssetPlacementEdModeToolkit::GetInlineContent() const
 {
-	return SNew(SVerticalBox)
-		+ SVerticalBox::Slot()
-		.FillHeight(.6)
+	return 	SNew(SScrollBox)
+		+ SScrollBox::Slot()
 		[
-			SNew(SScrollBox)
-			+ SScrollBox::Slot()
-			[
-				DetailsView.ToSharedRef()
-			]
-			+ SScrollBox::Slot()
-			[
-				ModeDetailsView.ToSharedRef()
-			]
+			DetailsView.ToSharedRef()
 		]
-		+ SVerticalBox::Slot()
-		.AutoHeight()
+		+ SScrollBox::Slot()
+		[
+			ModeDetailsView.ToSharedRef()
+		];
+}
+
+void FAssetPlacementEdModeToolkit::RequestModeUITabs()
+{
+	FModeToolkit::RequestModeUITabs();
+	if (ModeUILayer.IsValid())
+	{
+		TSharedPtr<FAssetEditorModeUILayer> ModeUILayerPtr = ModeUILayer.Pin();
+		TSharedRef<FWorkspaceItem> MenuGroup = ModeUILayerPtr->GetModeMenuCategory().ToSharedRef();
+		AssetPaletteInfo.OnSpawnTab = FOnSpawnTab::CreateSP(SharedThis(this), &FAssetPlacementEdModeToolkit::CreateAssetPalette);
+		AssetPaletteInfo.TabLabel = LOCTEXT("AssetPaletteTab", "Asset Palette");
+		AssetPaletteInfo.TabTooltip = LOCTEXT("ModesToolboxTabTooltipText", "Open the  Modes tab, which contains the active editor mode's settings.");
+		ModeUILayerPtr->SetModePanelInfo(FAssetEditorModeUILayer::BottomLeftTabID, AssetPaletteInfo);
+	}
+}
+
+
+
+TSharedRef<SDockTab> FAssetPlacementEdModeToolkit::CreateAssetPalette(const FSpawnTabArgs& Args)
+{
+	TSharedPtr<SDockTab> CreatedTab = SNew(SDockTab)
 		[
-			SNew(SExpandableArea)
-			.AreaTitle(LOCTEXT("AssetPaletteHeader", "Asset Palette"))
-			.BodyContent()
+			SNew(SScrollBox)
+			+ SScrollBox::Slot()
 			[
-				SNew(SScrollBox)
-				+ SScrollBox::Slot()
-				[
-					SNew(SAssetPlacementPalette)
-				]
+				SNew(SAssetPlacementPalette)
 			]
 		];
+
+	const FSlateBrush* TabIcon = GetEditorModeIcon().GetSmallIcon();
+	if (CreatedTab)
+	{
+		CreatedTab->SetTabIcon(TabIcon);
+	}
+	AssetPaletteTab = CreatedTab;
+	return CreatedTab.ToSharedRef();
 }
 
 #undef LOCTEXT_NAMESPACE
diff --git a/Engine/Source/Editor/AssetPlacementEdMode/Public/AssetPlacementEdModeToolkit.h b/Engine/Source/Editor/AssetPlacementEdMode/Public/AssetPlacementEdModeToolkit.h
--- a/Engine/Source/Editor/AssetPlacementEdMode/Public/AssetPlacementEdModeToolkit.h
+++ b/Engine/Source/Editor/AssetPlacementEdMode/Public/AssetPlacementEdModeToolkit.h
@@ -20,7 +20,13 @@ public:
 	// IToolkit interface 
 	virtual FName GetToolkitFName() const override;
 	virtual FText GetBaseToolkitName() const override;
-
+	virtual void InvokeUI() override;
 protected:
 	virtual TSharedPtr<SWidget> GetInlineContent() const override;
+	virtual void RequestModeUITabs() override;
+	TSharedRef<SDockTab> CreateAssetPalette(const FSpawnTabArgs& Args);
+
+private:
+	TWeakPtr<SDockTab> AssetPaletteTab;
+	FMinorTabConfig AssetPaletteInfo;
 };
