Performance improvement to brush based selection tools in asset placement mode.

I was able to get world sweeps to function with the ISMC by adding a small edge case for dealing with those types of components in the gather elements function. I think this is a good middle ground for now to keep world brush selection working without relying on the hit proxies or caching the elements from the drag.

diff --git a/Engine/Source/Editor/AssetPlacementEdMode/Private/Tools/PlacementBrushToolBase.cpp b/Engine/Source/Editor/AssetPlacementEdMode/Private/Tools/PlacementBrushToolBase.cpp
--- a/Engine/Source/Editor/AssetPlacementEdMode/Private/Tools/PlacementBrushToolBase.cpp
+++ b/Engine/Source/Editor/AssetPlacementEdMode/Private/Tools/PlacementBrushToolBase.cpp
@@ -22,6 +22,8 @@
 #include "Modes/PlacementModeSubsystem.h"
 #include "ActorFactories/ActorFactory.h"
 #include "BaseGizmos/GizmoRenderingUtil.h"
+#include "Tools/AssetEditorContextInterface.h"
+#include "ContextObjectStore.h"
 
 bool UPlacementToolBuilderBase::CanBuildTool(const FToolBuilderState& SceneState) const
 {	
@@ -265,55 +267,60 @@ FTypedElementListRef UPlacementBrushToolBase::GetElementsInBrushRadius(const FIn
 {
 	FTypedElementListRef ElementHandles = UTypedElementRegistry::GetInstance()->CreateElementList();
 
-	// We need the 2D device screen space position to test against hit proxies.
-	if (!DragPos.bHas2D)
+	const UTypedElementSelectionSet* SelectionSet = nullptr;
+	if (IAssetEditorContextInterface* AssetEditorAPI = GetToolManager()->GetContextObjectStore()->FindContext<IAssetEditorContextInterface>())
 	{
-		return ElementHandles;
+		SelectionSet = AssetEditorAPI->GetSelectionSet();
 	}
 
-	FViewport* Viewport = GetToolManager()->GetContextQueriesAPI()->GetFocusedViewport();
-	if (!Viewport)
+	if (!SelectionSet)
 	{
 		return ElementHandles;
 	}
 
-	FToolBuilderState SelectionState;
-	GetToolManager()->GetContextQueriesAPI()->GetCurrentSelectionState(SelectionState);
-	if (!SelectionState.TypedElementSelectionSet.IsValid())
-	{
-		return ElementHandles;
-	}
-	UTypedElementSelectionSet* SelectionSet = SelectionState.TypedElementSelectionSet.Get();
-
-	// Convert brush radius to screen space and gather hit elements within the brush radius.
-	int32 HalfBrushRadius = FMath::CeilToFloat((LastBrushStamp.Radius * LastBrushStampWorldToPixelScale) / 2.0f);
-	FIntRect AreaToCheck;
-	FIntPoint ViewportSize = Viewport->GetSizeXY();
-	AreaToCheck.Min.X = FMath::Max<int32>(0, DragPos.ScreenPosition.X - HalfBrushRadius);
-	AreaToCheck.Min.Y = FMath::Max<int32>(0, DragPos.ScreenPosition.Y - HalfBrushRadius);
-	AreaToCheck.Max.X = FMath::Min<int32>(ViewportSize.X, DragPos.ScreenPosition.X + HalfBrushRadius);
-	AreaToCheck.Max.Y = FMath::Min<int32>(ViewportSize.Y, DragPos.ScreenPosition.Y + HalfBrushRadius);
-
-	// Get the raw hit proxies within the rect
-	FTypedElementListRef HitElementHandles = UTypedElementRegistry::GetInstance()->CreateElementList();
-	Viewport->GetElementHandlesInRect(AreaToCheck, HitElementHandles);
-
-	// Work out which elements to actually select in the viewport, and also verify that we're actually intersecting with the sphere from the brush in world space.
-	FBoxSphereBounds WorldBrushSphereBounds(FSphere(LastBrushStamp.WorldPosition, LastBrushStamp.Radius));
-	HitElementHandles->ForEachElementHandle([SelectionSet, ElementHandles, &WorldBrushSphereBounds](const FTypedElementHandle& HitHandle)
+	FCollisionQueryParams QueryParams(TEXT("PlacementBrushTool"), SCENE_QUERY_STAT_ONLY(EdMode_PlacementTrace), true);
+	QueryParams.bReturnFaceIndex = false;
+	TArray<FHitResult> Hits;
+	FCollisionShape BrushSphere;
+	BrushSphere.SetSphere(LastBrushStamp.Radius);
+
+	const FVector TraceStart(LastWorldRay.Origin);
+	const FVector TraceEnd(LastWorldRay.Origin + LastWorldRay.Direction * HALF_WORLD_MAX);
+
+	GetToolManager()->GetWorld()->SweepMultiByObjectType(Hits, TraceStart, TraceEnd, FQuat::Identity,
+		FCollisionObjectQueryParams(FCollisionObjectQueryParams::InitType::AllObjects), BrushSphere, QueryParams);
+
+	for (const FHitResult& Hit : Hits)
 	{
-		FTypedElementHandle ResolvedHandle = SelectionSet->GetSelectionElement(HitHandle, ETypedElementSelectionMethod::Primary);
-		if (TTypedElement<UTypedElementWorldInterface> WorldInterface = SelectionSet->GetElementList()->GetElement<UTypedElementWorldInterface>(ResolvedHandle))
+		if (UInstancedStaticMeshComponent* ISMComponent = Cast<UInstancedStaticMeshComponent>(Hit.GetComponent()))
+		{
+			FTypedElementHandle StaticMeshObjectHandle = UEngineElementsLibrary::AcquireEditorObjectElementHandle(ISMComponent->GetStaticMesh());
+			if (!GEditor->GetEditorSubsystem<UPlacementModeSubsystem>()->DoesCurrentPaletteSupportElement(StaticMeshObjectHandle))
+			{
+				continue;
+			}
+
+			TArray<int32> OverlappingInstances = ISMComponent->GetInstancesOverlappingSphere(LastBrushStamp.WorldPosition, LastBrushStamp.Radius);
+			ElementHandles->Reserve(OverlappingInstances.Num());
+			for (const int32& InstanceIdx : OverlappingInstances)
+			{
+				if (FTypedElementHandle SMInstanceHandle = UEngineElementsLibrary::AcquireEditorSMInstanceElementHandle(ISMComponent, InstanceIdx))
+				{
+					ElementHandles->Add(MoveTemp(SMInstanceHandle));
+				}
+			}
+		}
+		else
 		{
-			FBoxSphereBounds ElementBounds;
-			WorldInterface.GetBounds(ElementBounds);
-			if (ElementBounds.SpheresIntersect(ElementBounds, WorldBrushSphereBounds))
+			FTypedElementHandle PrimarySelectionHandle = SelectionSet->GetSelectionElement(UEngineElementsLibrary::AcquireEditorComponentElementHandle(Hit.GetComponent()), ETypedElementSelectionMethod::Primary);
+			if (!GEditor->GetEditorSubsystem<UPlacementModeSubsystem>()->DoesCurrentPaletteSupportElement(PrimarySelectionHandle))
 			{
-				ElementHandles->Add(MoveTemp(ResolvedHandle));
+				continue;
 			}
+
+			ElementHandles->Add(MoveTemp(PrimarySelectionHandle));
 		}
-		return true;
-	});
+	}
 
 	return ElementHandles;
 }
diff --git a/Engine/Source/Editor/AssetPlacementEdMode/Private/Tools/PlacementLassoSelectTool.cpp b/Engine/Source/Editor/AssetPlacementEdMode/Private/Tools/PlacementLassoSelectTool.cpp
--- a/Engine/Source/Editor/AssetPlacementEdMode/Private/Tools/PlacementLassoSelectTool.cpp
+++ b/Engine/Source/Editor/AssetPlacementEdMode/Private/Tools/PlacementLassoSelectTool.cpp
@@ -34,33 +34,12 @@ UPlacementBrushToolBase* UPlacementModeLassoSelectToolBuilder::FactoryToolInstan
 void UPlacementModeLassoSelectTool::OnBeginDrag(const FRay& Ray)
 {
 	Super::OnBeginDrag(Ray);
-
-	ElementsFromDrag.Reset();
 	GetToolManager()->BeginUndoTransaction(NSLOCTEXT("AssetPlacementEdMode", "BrushSelect", "Select Elements"));
 }
 
 void UPlacementModeLassoSelectTool::OnEndDrag(const FRay& Ray)
 {
-	if (IAssetEditorContextInterface* AssetEditorContext = GetToolManager()->GetContextObjectStore()->FindContext<IAssetEditorContextInterface>())
-	{
-		UTypedElementSelectionSet* SelectionSet = AssetEditorContext->GetMutableSelectionSet();
-		if (SelectionSet && ElementsFromDrag)
-		{
-			const bool bSelectElements = !bCtrlToggle;
-			if (bSelectElements)
-			{
-				SelectionSet->SelectElements(ElementsFromDrag.ToSharedRef(), PlacementModeLassoToolInternal::SelectionOptions);
-			}
-			else
-			{
-				SelectionSet->DeselectElements(ElementsFromDrag.ToSharedRef(), PlacementModeLassoToolInternal::SelectionOptions);
-			}
-		}
-	}
-
 	GetToolManager()->EndUndoTransaction();
-	ElementsFromDrag.Reset();
-
 	Super::OnEndDrag(Ray);
 }
 
@@ -71,13 +50,21 @@ void UPlacementModeLassoSelectTool::OnTick(float DeltaTime)
 		return;
 	}
 
-	FTypedElementListRef HitElements = GetElementsInBrushRadius(LastDeviceInputRay);
-	if (ElementsFromDrag)
-	{
-		ElementsFromDrag->Append(HitElements);
-	}
-	else
+	if (IAssetEditorContextInterface* AssetEditorContext = GetToolManager()->GetContextObjectStore()->FindContext<IAssetEditorContextInterface>())
 	{
-		ElementsFromDrag = HitElements;
+		UTypedElementSelectionSet* SelectionSet = AssetEditorContext->GetMutableSelectionSet();
+		if (SelectionSet)
+		{
+			FTypedElementListRef HitElements = GetElementsInBrushRadius(LastDeviceInputRay);
+			const bool bSelectElements = !bCtrlToggle;
+			if (bSelectElements)
+			{
+				SelectionSet->SelectElements(HitElements, PlacementModeLassoToolInternal::SelectionOptions);
+			}
+			else
+			{
+				SelectionSet->DeselectElements(HitElements, PlacementModeLassoToolInternal::SelectionOptions);
+			}
+		}
 	}
 }
diff --git a/Engine/Source/Editor/AssetPlacementEdMode/Private/Tools/PlacementLassoSelectTool.h b/Engine/Source/Editor/AssetPlacementEdMode/Private/Tools/PlacementLassoSelectTool.h
--- a/Engine/Source/Editor/AssetPlacementEdMode/Private/Tools/PlacementLassoSelectTool.h
+++ b/Engine/Source/Editor/AssetPlacementEdMode/Private/Tools/PlacementLassoSelectTool.h
@@ -27,7 +27,4 @@ public:
 	virtual void OnTick(float DeltaTime) override;
 
 	constexpr static TCHAR ToolName[] = TEXT("LassoSelectTool");
-
-private:
-	FTypedElementListPtr ElementsFromDrag;
 };
