Initial implementation of palette asset and palette asset switching.

Also fixed an issue where the right click context menu wasn't spawning.

diff --git a/Engine/Source/Editor/AssetPlacementEdMode/Private/AssetPlacementEdMode.cpp b/Engine/Source/Editor/AssetPlacementEdMode/Private/AssetPlacementEdMode.cpp
--- a/Engine/Source/Editor/AssetPlacementEdMode/Private/AssetPlacementEdMode.cpp
+++ b/Engine/Source/Editor/AssetPlacementEdMode/Private/AssetPlacementEdMode.cpp
@@ -126,7 +126,7 @@ bool UAssetPlacementEdMode::IsSelectionAllowed(AActor* InActor, bool bInSelectio
 
 	// And we need to have a valid palette item.
 	FTypedElementHandle ActorHandle = UEngineElementsLibrary::AcquireEditorActorElementHandle(InActor);
-	return GEditor->GetEditorSubsystem<UPlacementModeSubsystem>()->DoesCurrentPaletteSupportElement(ActorHandle);
+	return GEditor->GetEditorSubsystem<UPlacementModeSubsystem>()->DoesActivePaletteSupportElement(ActorHandle);
 }
 
 void UAssetPlacementEdMode::OnToolStarted(UInteractiveToolManager* Manager, UInteractiveTool* Tool)
@@ -183,6 +183,12 @@ bool UAssetPlacementEdMode::ShouldDrawWidget() const
 		return false;
 	}
 
+	// Disable the widget in the lasso tool. The drag operations currently fight with the widget.
+	if (GetToolManager()->GetActiveToolName(EToolSide::Mouse) == UPlacementModeLassoSelectTool::ToolName)
+	{
+		return false;
+	}
+
 	return Super::ShouldDrawWidget();
 }
 
@@ -215,7 +221,7 @@ void UAssetPlacementEdMode::ClearSelection()
 
 bool UAssetPlacementEdMode::HasAnyAssetsInPalette() const
 {
-	return SettingsObjectAsPlacementSettings.IsValid() ? (SettingsObjectAsPlacementSettings->PaletteItems.Num() > 0) : false;
+	return SettingsObjectAsPlacementSettings.IsValid() ? (SettingsObjectAsPlacementSettings->GetActivePaletteItems().Num() > 0) : false;
 }
 
 bool UAssetPlacementEdMode::HasActiveSelection() const
diff --git a/Engine/Source/Editor/AssetPlacementEdMode/Private/AssetPlacementEdModeStyle.cpp b/Engine/Source/Editor/AssetPlacementEdMode/Private/AssetPlacementEdModeStyle.cpp
--- a/Engine/Source/Editor/AssetPlacementEdMode/Private/AssetPlacementEdModeStyle.cpp
+++ b/Engine/Source/Editor/AssetPlacementEdMode/Private/AssetPlacementEdModeStyle.cpp
@@ -43,7 +43,9 @@ FAssetPlacementEdModeStyle::FAssetPlacementEdModeStyle()
 	SetContentRoot(FPaths::EngineContentDir() / TEXT("Editor/Slate"));
 	SetCoreContentRoot(FPaths::EngineContentDir() / TEXT("Slate"));
 
-	Set("LevelEditor.AssetPlacementEdMode", new IMAGE_BRUSH_SVG("Starship/MainToolbar/AssetPlacementMode", FVector2D(20.0f, 20.0f)));
+	Set("LevelEditor.AssetPlacementEdMode", new IMAGE_BRUSH_SVG("Starship/MainToolbar/AssetPlacementMode", Icon20x20));
+	Set("ClassIcon.PlacementPaletteAsset", new IMAGE_BRUSH_SVG("Starship/MainToolbar/AssetPlacementMode", FVector2D(16.0f, 16.0f)));
+	Set("ClassThumbnail.PlacementPaletteAsset", new IMAGE_BRUSH_SVG("Starship/MainToolbar/AssetPlacementMode", FVector2D(64.0f, 64.0f)));
 
 	Set("AssetPlacementEdMode.Select",                      new IMAGE_BRUSH("Icons/GeneralTools/Select_40x", Icon20x20));
 	Set("AssetPlacementEdMode.Select.Small",                new IMAGE_BRUSH("Icons/GeneralTools/Select_40x", Icon20x20));
diff --git a/Engine/Source/Editor/AssetPlacementEdMode/Private/AssetPlacementPaletteItem.cpp b/Engine/Source/Editor/AssetPlacementEdMode/Private/AssetPlacementPaletteItem.cpp
--- a/Engine/Source/Editor/AssetPlacementEdMode/Private/AssetPlacementPaletteItem.cpp
+++ b/Engine/Source/Editor/AssetPlacementEdMode/Private/AssetPlacementPaletteItem.cpp
@@ -20,9 +20,9 @@
 
 #define LOCTEXT_NAMESPACE "AssetPlacementMode"
 
-FAssetPlacementPaletteItemModel::FAssetPlacementPaletteItemModel(FAssetPlacementUIInfoPtr InTypeInfo, TSharedRef<SAssetPlacementPalette> InFoliagePalette, TSharedPtr<FAssetThumbnailPool> InThumbnailPool/*, FEdModeFoliage* InFoliageEditMode*/)
-	: TypeInfo(InTypeInfo)
-	, AssetPalette(InFoliagePalette)
+FAssetPlacementPaletteItemModel::FAssetPlacementPaletteItemModel(const FAssetData& InAssetData, TSharedRef<SAssetPlacementPalette> InParentPalette, TSharedPtr<FAssetThumbnailPool> InThumbnailPool)
+	: TypeInfo(MakeShared<FPaletteItemUIInfo>(InAssetData))
+	, AssetPalette(InParentPalette)
 {
 	check(TypeInfo);
 	DisplayFName = TypeInfo->AssetData.AssetName;
@@ -49,11 +49,6 @@ FAssetPlacementPaletteItemModel::FAssetPlacementPaletteItemModel(FAssetPlacement
 	ThumbnailWidget = Thumbnail->MakeThumbnailWidget(ThumbnailConfig);
 }
 
-TSharedPtr<SAssetPlacementPalette> FAssetPlacementPaletteItemModel::GetAssetPalette() const
-{
-	return AssetPalette.Pin();
-}
-
 FAssetPlacementUIInfoPtr FAssetPlacementPaletteItemModel::GetTypeUIInfo() const
 {
 	return TypeInfo;
diff --git a/Engine/Source/Editor/AssetPlacementEdMode/Private/AssetPlacementPaletteItem.h b/Engine/Source/Editor/AssetPlacementEdMode/Private/AssetPlacementPaletteItem.h
--- a/Engine/Source/Editor/AssetPlacementEdMode/Private/AssetPlacementPaletteItem.h
+++ b/Engine/Source/Editor/AssetPlacementEdMode/Private/AssetPlacementPaletteItem.h
@@ -11,6 +11,7 @@
 #include "Widgets/Views/STableViewBase.h"
 #include "Widgets/Views/STableRow.h"
 #include "SAssetPlacementPalette.h"
+#include "AssetRegistry/AssetData.h"
 
 class FAssetThumbnailPool;
 struct FPaletteItem;
@@ -21,7 +22,18 @@ namespace PlacementPaletteConstants
 	const FInt32Interval ThumbnailSizeRange(32, 128);
 }
 
-typedef TSharedPtr<FPaletteItem> FAssetPlacementUIInfoPtr;
+struct FPaletteItemUIInfo
+{
+	FPaletteItemUIInfo(const FAssetData& InAssetData)
+		: AssetData(InAssetData)
+	{
+		check(AssetData.IsValid());
+	}
+
+	const FAssetData AssetData;
+};
+
+typedef TSharedPtr<FPaletteItemUIInfo> FAssetPlacementUIInfoPtr;
 
 namespace AssetPlacementPaletteTreeColumns
 {
@@ -31,10 +43,7 @@ namespace AssetPlacementPaletteTreeColumns
 class FAssetPlacementPaletteItemModel : public TSharedFromThis<FAssetPlacementPaletteItemModel>
 {
 public:
-	FAssetPlacementPaletteItemModel(FAssetPlacementUIInfoPtr InTypeInfo, TSharedRef<class SAssetPlacementPalette> InFoliagePalette, TSharedPtr<class FAssetThumbnailPool> InThumbnailPool);
-
-	/** @return The foliage palette that contains the item */
-	TSharedPtr<SAssetPlacementPalette> GetAssetPalette() const;
+	FAssetPlacementPaletteItemModel(const FAssetData& InAssetData, TSharedRef<class SAssetPlacementPalette> InParentPalette, TSharedPtr<class FAssetThumbnailPool> InThumbnailPool);
 
 	FAssetPlacementUIInfoPtr GetTypeUIInfo() const;
 
@@ -66,8 +75,6 @@ private:
 	/** Gets the source asset type text */
 	FText GetSourceAssetTypeText() const;
 
-	TWeakPtr<IAssetTypeActions> GetAssetTypeActions() const;
-
 private:
 	TSharedPtr<SWidget> ThumbnailWidget;
 	FName DisplayFName;
diff --git a/Engine/Source/Editor/AssetPlacementEdMode/Private/AssetPlacementSettings.cpp b/Engine/Source/Editor/AssetPlacementEdMode/Private/AssetPlacementSettings.cpp
--- a/Engine/Source/Editor/AssetPlacementEdMode/Private/AssetPlacementSettings.cpp
+++ b/Engine/Source/Editor/AssetPlacementEdMode/Private/AssetPlacementSettings.cpp
@@ -1,6 +1,9 @@
 // Copyright Epic Games, Inc. All Rights Reserved.
 
 #include "AssetPlacementSettings.h"
+#include "PlacementPaletteAsset.h"
+#include "PackageTools.h"
+#include "PlacementPaletteItem.h"
 
 bool UAssetPlacementSettings::CanEditChange(const FProperty* InProperty) const
 {
@@ -37,3 +40,66 @@ bool UAssetPlacementSettings::CanEditChange(const FProperty* InProperty) const
 
 	return true;
 }
+
+void UAssetPlacementSettings::SetPaletteAsset(UPlacementPaletteAsset* InPaletteAsset)
+{
+	if (!InPaletteAsset)
+	{
+		ActivePalette = UserPalette;
+	}
+	else
+	{
+		ActivePalette = InPaletteAsset;
+	}
+
+	LastActivePalettePath = FSoftObjectPath(InPaletteAsset);
+}
+
+void UAssetPlacementSettings::AddItemToActivePalette(const FPaletteItem& InPaletteItem)
+{
+	ActivePalette->Modify();
+	ActivePalette->PaletteItems.Add(InPaletteItem);
+}
+
+TArrayView<const FPaletteItem> UAssetPlacementSettings::GetActivePaletteItems() const
+{
+	return ActivePalette ? MakeArrayView(ActivePalette->PaletteItems) : TArrayView<const FPaletteItem>();
+}
+
+void UAssetPlacementSettings::ClearActivePaletteItems()
+{
+	ActivePalette->PaletteItems.Empty();
+}
+
+void UAssetPlacementSettings::Restore()
+{
+	LoadConfig();
+
+	UserPalette = NewObject<UPlacementPaletteAsset>(this);
+	UserPalette->LoadConfig();
+
+	ActivePalette = Cast<UPlacementPaletteAsset>(LastActivePalettePath.TryLoad());
+	if (!ActivePalette)
+	{
+		ActivePalette = UserPalette;
+	}
+}
+
+void UAssetPlacementSettings::SaveActivePalette()
+{
+	if (ActivePalette != UserPalette)
+	{
+		UPackageTools::SavePackagesForObjects(TArray<UObject*>({ ActivePalette }));
+	}
+	else
+	{
+		UserPalette->SaveConfig();
+	}
+}
+
+void UAssetPlacementSettings::Save()
+{
+	UPackageTools::SavePackagesForObjects(TArray<UObject*>({ActivePalette}));
+	UserPalette->SaveConfig();
+	SaveConfig();
+}
diff --git a/Engine/Source/Editor/AssetPlacementEdMode/Private/AssetPlacementSettings.h b/Engine/Source/Editor/AssetPlacementEdMode/Private/AssetPlacementSettings.h
--- a/Engine/Source/Editor/AssetPlacementEdMode/Private/AssetPlacementSettings.h
+++ b/Engine/Source/Editor/AssetPlacementEdMode/Private/AssetPlacementSettings.h
@@ -3,24 +3,12 @@
 #pragma once
 
 #include "CoreMinimal.h"
-#include "AssetRegistry/AssetData.h"
 #include "FoliageType.h"
 
 #include "AssetPlacementSettings.generated.h"
 
-class IAssetFactoryInterface;
-
-USTRUCT()
-struct FPaletteItem
-{
-	GENERATED_BODY()
-
-	UPROPERTY()
-	FAssetData AssetData;
-
-	UPROPERTY()
-	TScriptInterface<IAssetFactoryInterface> AssetFactoryInterface;
-};
+class UPlacementPaletteAsset;
+struct FPaletteItem;
 
 UCLASS(config = EditorPerProjectUserSettings)
 class UAssetPlacementSettings : public UObject
@@ -128,12 +116,27 @@ public:
 	UPROPERTY(config, EditAnywhere, Category = "Place on Filters")
 	bool bTranslucent = false;
 
-	// todo: palette save and load between sessions
-	//UPROPERTY(config)
-	TArray<TSharedPtr<FPaletteItem>> PaletteItems;
-
 	UPROPERTY(config)
 	bool bUseContentBrowserSelection = true;
 
 	virtual bool CanEditChange(const FProperty* InProperty) const override;
+
+	void SetPaletteAsset(UPlacementPaletteAsset* InPaletteAsset);
+	void AddItemToActivePalette(const FPaletteItem& InPaletteItem);
+	TArrayView<const FPaletteItem> GetActivePaletteItems() const;
+	void ClearActivePaletteItems();
+	void SaveActivePalette();
+	void Restore();
+	void Save();
+
+protected:
+
+	UPROPERTY(Transient)
+	TObjectPtr<UPlacementPaletteAsset> ActivePalette;
+
+	UPROPERTY(Transient)
+	TObjectPtr<UPlacementPaletteAsset> UserPalette;
+
+	UPROPERTY(config)
+	FSoftObjectPath LastActivePalettePath;
 };
diff --git a/Engine/Source/Editor/AssetPlacementEdMode/Private/Modes/PlacementModeSubsystem.cpp b/Engine/Source/Editor/AssetPlacementEdMode/Private/Modes/PlacementModeSubsystem.cpp
--- a/Engine/Source/Editor/AssetPlacementEdMode/Private/Modes/PlacementModeSubsystem.cpp
+++ b/Engine/Source/Editor/AssetPlacementEdMode/Private/Modes/PlacementModeSubsystem.cpp
@@ -7,6 +7,8 @@
 #include "Elements/Framework/TypedElementRegistry.h"
 #include "Elements/Interfaces/TypedElementAssetDataInterface.h"
 #include "Factories/AssetFactoryInterface.h"
+#include "PlacementPaletteAsset.h"
+#include "PlacementPaletteItem.h"
 
 #include "Subsystems/PlacementSubsystem.h"
 #include "Editor.h"
@@ -14,12 +16,12 @@
 void UPlacementModeSubsystem::Initialize(FSubsystemCollectionBase& Collection)
 {
 	ModeSettings = NewObject<UAssetPlacementSettings>(this);
-	ModeSettings->LoadConfig();
+	ModeSettings->Restore();
 }
 
 void UPlacementModeSubsystem::Deinitialize()
 {
-	ModeSettings->SaveConfig();
+	ModeSettings->Save();
 	ModeSettings = nullptr;
 }
 
@@ -28,52 +30,44 @@ const UAssetPlacementSettings* UPlacementModeSubsystem::GetModeSettingsObject()
 	return ModeSettings;
 }
 
-bool UPlacementModeSubsystem::DoesCurrentPaletteSupportElement(const FTypedElementHandle& InElementToCheck) const
+UAssetPlacementSettings* UPlacementModeSubsystem::GetMutableModeSettingsObject()
 {
-	if (!ModeSettings)
-	{
-		return false;
-	}
+	return ModeSettings;
+}
 
+bool UPlacementModeSubsystem::DoesActivePaletteSupportElement(const FTypedElementHandle& InElementToCheck) const
+{
 	if (TTypedElement<UTypedElementAssetDataInterface> AssetDataInterface = UTypedElementRegistry::GetInstance()->GetElement<UTypedElementAssetDataInterface>(InElementToCheck))
 	{
 		TArray<FAssetData> ReferencedAssetDatas = AssetDataInterface.GetAllReferencedAssetDatas();
-		for (const TSharedPtr<FPaletteItem>& Item : ModeSettings->PaletteItems)
+		for (const FPaletteItem& Item : ModeSettings->GetActivePaletteItems())
 		{
-			if (!Item)
+			if (ReferencedAssetDatas.FindByPredicate([&Item](const FAssetData& ReferencedAssetData){ return (ReferencedAssetData.ToSoftObjectPath() == Item.AssetPath); }))
+			{
+				return true;
+			}
+
+			// The current implementation of the asset data interface for actors requires that individual actors report on assets contained within their components.
+			// Not all actors do this reliably, so additionally check the supplied factory for a match. 
+			if (!Item.AssetFactoryInterface)
 			{
 				continue;
 			}
 
-			if (ReferencedAssetDatas.Find(Item->AssetData) != INDEX_NONE)
+			FAssetData FoundAssetDataFromFactory = Item.AssetFactoryInterface->GetAssetDataFromElementHandle(InElementToCheck);
+			if (FoundAssetDataFromFactory.ToSoftObjectPath() == Item.AssetPath)
 			{
 				return true;
 			}
 		}
 	}
 
-	// The current implementation of the asset data interface for actors requires that individual actors report on assets contained within their components.
-	// Not all actors do this reliably, so additionally check the supplied factory for a match. 
-	for (const TSharedPtr<FPaletteItem>& Item : ModeSettings->PaletteItems)
-	{
-		if (!Item || !Item->AssetFactoryInterface)
-		{
-			continue;
-		}
-
-		FAssetData FoundAssetDataFromFactory = Item->AssetFactoryInterface->GetAssetDataFromElementHandle(InElementToCheck);
-		if (FoundAssetDataFromFactory == Item->AssetData)
-		{
-			return true;
-		}
-	}
-
 	return false;
 }
 
-TSharedPtr<FPaletteItem> UPlacementModeSubsystem::AddPaletteItem(const FAssetData& InAssetData)
+FPaletteItem UPlacementModeSubsystem::CreatePaletteItem(const FAssetData& InAssetData)
 {
-	TSharedPtr<FPaletteItem> NewPaletteItem;
+	FPaletteItem NewPaletteItem;
 	if (!InAssetData.IsValid())
 	{
 		return NewPaletteItem;
@@ -84,16 +78,16 @@ TSharedPtr<FPaletteItem> UPlacementModeSubsystem::AddPaletteItem(const FAssetDat
 		return NewPaletteItem;
 	}
 
-	if (ModeSettings && !ModeSettings->PaletteItems.FindByPredicate([InAssetData](const TSharedPtr<FPaletteItem>& ItemIter) { return ItemIter ? (ItemIter->AssetData.ObjectPath == InAssetData.ObjectPath) : false; }))
+	if (!ModeSettings->GetActivePaletteItems().FindByPredicate([InAssetData](const FPaletteItem& ItemIter) { return (ItemIter.AssetPath == InAssetData.ToSoftObjectPath()); }))
 	{
 		if (UPlacementSubsystem* PlacementSubystem = GEditor->GetEditorSubsystem<UPlacementSubsystem>())
 		{
 			if (TScriptInterface<IAssetFactoryInterface> AssetFactory = PlacementSubystem->FindAssetFactoryFromAssetData(InAssetData))
 			{
-				NewPaletteItem = MakeShared<FPaletteItem>();
-				NewPaletteItem->AssetData = InAssetData;
-				NewPaletteItem->AssetFactoryInterface = AssetFactory;
-				ModeSettings->PaletteItems.Add(NewPaletteItem);
+				NewPaletteItem.ItemGuid = FGuid::NewGuid();
+				NewPaletteItem.AssetPath = InAssetData.ToSoftObjectPath();
+				NewPaletteItem.AssetFactoryInterface = AssetFactory;
+				ModeSettings->AddItemToActivePalette(NewPaletteItem);
 			}
 		}
 	}
@@ -101,11 +95,6 @@ TSharedPtr<FPaletteItem> UPlacementModeSubsystem::AddPaletteItem(const FAssetDat
 	return NewPaletteItem;
 }
 
-void UPlacementModeSubsystem::ClearPalette()
-{
-	ModeSettings->PaletteItems.Empty();
-}
-
 void UPlacementModeSubsystem::SetUseContentBrowserAsPalette(bool bInUseContentBrowser)
 {
 	ModeSettings->bUseContentBrowserSelection = bInUseContentBrowser;
diff --git a/Engine/Source/Editor/AssetPlacementEdMode/Private/Modes/PlacementModeSubsystem.h b/Engine/Source/Editor/AssetPlacementEdMode/Private/Modes/PlacementModeSubsystem.h
--- a/Engine/Source/Editor/AssetPlacementEdMode/Private/Modes/PlacementModeSubsystem.h
+++ b/Engine/Source/Editor/AssetPlacementEdMode/Private/Modes/PlacementModeSubsystem.h
@@ -10,6 +10,7 @@
 struct FPaletteItem;
 struct FTypedElementHandle;
 class UAssetPlacementSettings;
+class UPlacementPaletteAsset;
 
 UCLASS(Transient)
 class UPlacementModeSubsystem : public UEditorSubsystem
@@ -24,25 +25,21 @@ public:
 
 	// @returns the settings object for the mode for sharing across all tools and tool builders.
 	const UAssetPlacementSettings* GetModeSettingsObject() const;
+	UAssetPlacementSettings* GetMutableModeSettingsObject();
 
 	/**
 	 * Verifies if the given element handle is supported by the current mode settings' palette.
 	 *
 	 * @returns true if the element can be placed by the mode.
 	 */
-	bool DoesCurrentPaletteSupportElement(const FTypedElementHandle& InElementToCheck) const;
+	bool DoesActivePaletteSupportElement(const FTypedElementHandle& InElementToCheck) const;
 
 	/**
 	 * Adds the given palette item to the current palette.
 	 * 
 	 * @returns a pointer to the added item, if the add was successful.
 	 */
-	TSharedPtr<FPaletteItem> AddPaletteItem(const FAssetData& InAssetData);
-
-	/**
-	 * Clears all items from the current palette.
-	 */
-	void ClearPalette();
+	FPaletteItem CreatePaletteItem(const FAssetData& InAssetData);
 
 	/**
 	 * Updates the settings object to use the content browser's active selection as the palette.
diff --git a/Engine/Source/Editor/AssetPlacementEdMode/Private/PlacementPaletteAsset.cpp b/Engine/Source/Editor/AssetPlacementEdMode/Private/PlacementPaletteAsset.cpp
new file mode 100644
--- /dev/null
+++ b/Engine/Source/Editor/AssetPlacementEdMode/Private/PlacementPaletteAsset.cpp
@@ -0,0 +1,21 @@
+// Copyright Epic Games, Inc. All Rights Reserved.
+
+#include "PlacementPaletteAsset.h"
+
+UPlacementPaletteAssetFactory::UPlacementPaletteAssetFactory(const FObjectInitializer& ObjectInitializer)
+	: Super(ObjectInitializer)
+{
+	SupportedClass = UPlacementPaletteAsset::StaticClass();
+	bCreateNew = true;
+	bEditAfterNew = true;
+}
+
+UObject* UPlacementPaletteAssetFactory::FactoryCreateNew(UClass* InClass, UObject* InParent, FName InName, EObjectFlags Flags, UObject* Context, FFeedbackContext* Warn)
+{
+	return NewObject<UPlacementPaletteAsset>(InParent, InClass, InName, Flags, Context);
+}
+
+bool UPlacementPaletteAssetFactory::ShouldShowInNewMenu() const
+{
+	return true;
+}
diff --git a/Engine/Source/Editor/AssetPlacementEdMode/Private/PlacementPaletteAsset.h b/Engine/Source/Editor/AssetPlacementEdMode/Private/PlacementPaletteAsset.h
new file mode 100644
--- /dev/null
+++ b/Engine/Source/Editor/AssetPlacementEdMode/Private/PlacementPaletteAsset.h
@@ -0,0 +1,31 @@
+// Copyright Epic Games, Inc. All Rights Reserved.
+
+#pragma once
+
+#include "Engine/DataAsset.h"
+#include "Factories/Factory.h"
+#include "PlacementPaletteItem.h"
+
+#include "PlacementPaletteAsset.generated.h"
+
+class UPlacementPaletteAsset;
+
+UCLASS(NotPlaceable, config = EditorPerProjectUserSettings)
+class UPlacementPaletteAsset : public UDataAsset
+{
+	GENERATED_BODY()
+
+public:
+	UPROPERTY(config, EditAnywhere, Category=Items)
+	TArray<FPaletteItem> PaletteItems;
+};
+
+UCLASS(hideCategories=Object)
+class UPlacementPaletteAssetFactory : public UFactory
+{
+	GENERATED_UCLASS_BODY()
+
+public:
+	virtual UObject* FactoryCreateNew(UClass* InClass, UObject* InParent, FName InName, EObjectFlags Flags, UObject* Context, FFeedbackContext* Warn) override;
+	virtual bool ShouldShowInNewMenu() const override;
+};
diff --git a/Engine/Source/Editor/AssetPlacementEdMode/Private/PlacementPaletteItem.h b/Engine/Source/Editor/AssetPlacementEdMode/Private/PlacementPaletteItem.h
new file mode 100644
--- /dev/null
+++ b/Engine/Source/Editor/AssetPlacementEdMode/Private/PlacementPaletteItem.h
@@ -0,0 +1,24 @@
+// Copyright Epic Games, Inc. All Rights Reserved.
+
+#pragma once
+
+#include "CoreMinimal.h"
+
+#include "PlacementPaletteItem.generated.h"
+
+class IAssetFactoryInterface;
+
+USTRUCT()
+struct FPaletteItem
+{
+	GENERATED_BODY()
+
+	UPROPERTY(EditAnywhere, Category = Placement)
+	FSoftObjectPath AssetPath;
+
+	UPROPERTY(EditAnywhere, Category = Placement, AdvancedDisplay)
+	TScriptInterface<IAssetFactoryInterface> AssetFactoryInterface;
+
+	UPROPERTY()
+	FGuid ItemGuid;
+};
diff --git a/Engine/Source/Editor/AssetPlacementEdMode/Private/SAssetPlacementPalette.cpp b/Engine/Source/Editor/AssetPlacementEdMode/Private/SAssetPlacementPalette.cpp
--- a/Engine/Source/Editor/AssetPlacementEdMode/Private/SAssetPlacementPalette.cpp
+++ b/Engine/Source/Editor/AssetPlacementEdMode/Private/SAssetPlacementPalette.cpp
@@ -17,7 +17,6 @@
 #include "IContentBrowserSingleton.h"
 #include "ContentBrowserModule.h"
 #include "PropertyCustomizationHelpers.h"
-#include "IDetailsView.h"
 #include "AssetData.h"
 #include "AssetSelection.h"
 #include "Editor.h"
@@ -37,6 +36,9 @@
 #include "ContentBrowserModule.h"
 #include "Modules/ModuleManager.h"
 #include "SAssetDropTarget.h"
+#include "PlacementPaletteAsset.h"
+#include "AssetRegistry/AssetRegistryModule.h"
+#include "PlacementPaletteItem.h"
 
 #define LOCTEXT_NAMESPACE "AssetPlacementMode"
 
@@ -59,21 +61,7 @@ void SAssetPlacementPalette::Construct(const FArguments& InArgs)
 
 	const FText BlankText = FText::GetEmpty();
 
-	TWeakObjectPtr<const UAssetPlacementSettings> PlacementSettings = GEditor->GetEditorSubsystem<UPlacementModeSubsystem>()->GetModeSettingsObject();
-	if (PlacementSettings.IsValid())
-	{
-		for (const TSharedPtr<FPaletteItem>& PaletteItem : PlacementSettings->PaletteItems)
-		{
-			if (PaletteItem)
-			{
-				PaletteItems.Add(MakeShared<FAssetPlacementPaletteItemModel>(PaletteItem, SharedThis(this), ThumbnailPool));
-			}
-		}
-
-		// Make sure the content browser setting is not what we're going to change to, so that the setup will happen.
-		bIsMirroringContentBrowser = !PlacementSettings->bUseContentBrowserSelection;
-		SetupContentBrowserMirroring(PlacementSettings->bUseContentBrowserSelection);
-	}
+	SetToLocalUserPalette();
 
 	ChildSlot
 	[
@@ -96,6 +84,7 @@ void SAssetPlacementPalette::Construct(const FArguments& InArgs)
 					.Type(ESlateCheckBoxType::Type::ToggleButton)
 					.IsChecked(bIsMirroringContentBrowser ? ECheckBoxState::Checked : ECheckBoxState::Unchecked)
 					.Style(&FAssetPlacementEdModeStyle::Get().GetWidgetStyle<FCheckBoxStyle>("ToggleButtonCheckBox"))
+					.Visibility(this, &SAssetPlacementPalette::GetContentBrowserMirrorVisibility)
 					.OnCheckStateChanged(this, &SAssetPlacementPalette::OnContentBrowserMirrorButtonClicked)
 					[
 						SNew(STextBlock)
@@ -105,6 +94,37 @@ void SAssetPlacementPalette::Construct(const FArguments& InArgs)
 						.ToolTipText(LOCTEXT("Placement_ToggleContentBrowserMirroring_ToolTip", "Toggles palette to mirror the active content browser selection."))
 					]
 				]
+				+SVerticalBox::Slot()
+				.AutoHeight()
+				[
+					SNew(SObjectPropertyEntryBox)
+					.AllowedClass(UPlacementPaletteAsset::StaticClass())
+					.ObjectPath(this, &SAssetPlacementPalette::GetPalettePath)
+					.OnObjectChanged(this, &SAssetPlacementPalette::OnSetPaletteAsset)
+					.ThumbnailPool(ThumbnailPool)
+					.CustomContentSlot()
+					[
+						SNew(SBox)
+						.HAlign(HAlign_Left)
+						.VAlign(VAlign_Center)
+						[
+							SNew(SHorizontalBox)
+							+ SHorizontalBox::Slot()
+							.VAlign(VAlign_Center)
+							.AutoWidth()
+							[
+								PropertyCustomizationHelpers::MakeClearButton(FSimpleDelegate::CreateSP(this, &SAssetPlacementPalette::OnResetPaletteAssetClicked), LOCTEXT("ResetPaletteAssetTooltip", "Clear the current palette asset, and use the user's local palette."))
+							]
+							+ SHorizontalBox::Slot()
+							.VAlign(VAlign_Center)
+							.AutoWidth()
+							.Padding(2.0f, 0.0f)
+							[
+								PropertyCustomizationHelpers::MakeSaveButton(FSimpleDelegate::CreateSP(this, &SAssetPlacementPalette::OnSavePaletteAssetClicked), LOCTEXT("SaveAssetPaletteTooltip", "Save changes to the current palette asset"))
+							]
+						]
+					]
+				]
 
 				+ SVerticalBox::Slot()
 				.VAlign(VAlign_Center)
@@ -165,6 +185,7 @@ void SAssetPlacementPalette::Construct(const FArguments& InArgs)
 				SNew(SAssetDropTarget)
 				.OnAreAssetsAcceptableForDrop(this, &SAssetPlacementPalette::OnAreAssetsValidForDrop)
 				.OnAssetsDropped(this, &SAssetPlacementPalette::HandlePlacementDropped)
+				.Visibility(this, &SAssetPlacementPalette::GetDropHintTextVisibility)
 				.bSupportsMultiDrop(true)
 				[
 					SNew(SBox)
@@ -173,7 +194,6 @@ void SAssetPlacementPalette::Construct(const FArguments& InArgs)
 					[
 						SNew(STextBlock)
 						.Text(LOCTEXT("Placement_AddPlacementMesh", "Drop Assets Here"))
-						.Visibility(this, &SAssetPlacementPalette::GetDropHintTextVisibility)
 						.ShadowOffset(FVector2D(1.f, 1.f))
 					]
 				]
@@ -234,8 +254,8 @@ void SAssetPlacementPalette::RefreshActivePaletteViewWidget()
 void SAssetPlacementPalette::AddPlacementType(const FAssetData& AssetData)
 {
 	// Try to add the item to the mode's palette
-	FAssetPlacementUIInfoPtr PlacementInfo = GEditor->GetEditorSubsystem<UPlacementModeSubsystem>()->AddPaletteItem(AssetData);
-	if (!PlacementInfo.IsValid())
+	FPaletteItem PlacementInfo = GEditor->GetEditorSubsystem<UPlacementModeSubsystem>()->CreatePaletteItem(AssetData);
+	if (!PlacementInfo.AssetPath.IsValid())
 	{
 		return;
 	}
@@ -243,14 +263,14 @@ void SAssetPlacementPalette::AddPlacementType(const FAssetData& AssetData)
 	// Try to load the asset async so it's ready to place.
 	UAssetManager::GetStreamableManager().RequestAsyncLoad(AssetData.ToSoftObjectPath());
 
-	PaletteItems.Add(MakeShared<FAssetPlacementPaletteItemModel>(PlacementInfo, SharedThis(this), ThumbnailPool));
+	PaletteItems.Add(MakeShared<FAssetPlacementPaletteItemModel>(AssetData, SharedThis(this), ThumbnailPool));
 }
 
 void SAssetPlacementPalette::ClearPalette()
 {
 	if (UPlacementModeSubsystem* ModeSubsystem = GEditor->GetEditorSubsystem<UPlacementModeSubsystem>())
 	{
-		ModeSubsystem->ClearPalette();
+		ModeSubsystem->GetMutableModeSettingsObject()->ClearActivePaletteItems();
 	}
 
 	PaletteItems.Empty();
@@ -361,6 +381,26 @@ bool SAssetPlacementPalette::ShouldFilterAsset(const FAssetData& InAssetData)
 	return true;
 }
 
+void SAssetPlacementPalette::OnResetPaletteAssetClicked()
+{
+	PalettePath.Reset();
+	SetToLocalUserPalette();
+	UpdatePalette(true);
+}
+
+void SAssetPlacementPalette::OnSavePaletteAssetClicked()
+{
+	if (UPlacementModeSubsystem* PlacementModeSubsystem = GEditor->GetEditorSubsystem<UPlacementModeSubsystem>())
+	{
+		PlacementModeSubsystem->GetMutableModeSettingsObject()->SaveActivePalette();
+	}
+}
+
+EVisibility SAssetPlacementPalette::GetContentBrowserMirrorVisibility() const
+{
+	return GetPalettePath().IsEmpty() ? EVisibility::Visible : EVisibility::Collapsed;
+}
+
 void SAssetPlacementPalette::OnContentBrowserMirrorButtonClicked(ECheckBoxState InState)
 {
 	SetupContentBrowserMirroring((InState == ECheckBoxState::Checked));
@@ -403,6 +443,72 @@ void SAssetPlacementPalette::SetupContentBrowserMirroring(bool bInMirrorContentB
 	bIsMirroringContentBrowser = bInMirrorContentBrowser;
 }
 
+FString SAssetPlacementPalette::GetPalettePath() const
+{
+	return PalettePath.ToString();
+}
+
+void SAssetPlacementPalette::OnSetPaletteAsset(const FAssetData& InAssetData)
+{
+	FSoftObjectPath NewAssetPath = InAssetData.ToSoftObjectPath();
+	if (NewAssetPath.IsValid() && (NewAssetPath != PalettePath))
+	{
+		if (UPlacementPaletteAsset* PaletteAsset = Cast<UPlacementPaletteAsset>(InAssetData.GetAsset()))
+		{
+			if (UPlacementModeSubsystem* PlacementModeSubsystem = GEditor->GetEditorSubsystem<UPlacementModeSubsystem>())
+			{
+				PlacementModeSubsystem->GetMutableModeSettingsObject()->SetPaletteAsset(PaletteAsset);
+
+				PalettePath = NewAssetPath;
+
+				// Shutdown content browser mirroring if needed
+				SetupContentBrowserMirroring(false);
+
+				SetPaletteItems(PlacementModeSubsystem->GetModeSettingsObject()->GetActivePaletteItems());
+			}
+		}
+	}
+}
+
+void SAssetPlacementPalette::SetPaletteItems(TArrayView<const FPaletteItem> InPaletteItems)
+{
+	PaletteItems.Empty();
+	FilteredItems.Empty();
+
+	FAssetRegistryModule& AssetRegistry = FModuleManager::Get().LoadModuleChecked<FAssetRegistryModule>(TEXT("AssetRegistry"));
+	TArray<FSoftObjectPath> AssetsToLoad;
+	for (const FPaletteItem& PaletteItem : InPaletteItems)
+	{
+		FAssetData AssetData = AssetRegistry.Get().GetAssetByObjectPath(PaletteItem.AssetPath.GetAssetPathName());
+		if (AssetData.IsValid())
+		{
+			PaletteItems.Add(MakeShared<FAssetPlacementPaletteItemModel>(AssetData, SharedThis(this), ThumbnailPool));
+			AssetsToLoad.Emplace(PaletteItem.AssetPath);
+		}
+	}
+
+	// Try to load the assets async so they're ready to place later
+	UAssetManager::GetStreamableManager().RequestAsyncLoad(AssetsToLoad);
+
+	// Update the palette view to match contained items
+	UpdatePalette(true);
+}
+
+void SAssetPlacementPalette::SetToLocalUserPalette()
+{
+	if (UPlacementModeSubsystem* PlacementModeSubsystem = GEditor->GetEditorSubsystem<UPlacementModeSubsystem>())
+	{
+		PlacementModeSubsystem->GetMutableModeSettingsObject()->SetPaletteAsset(nullptr);
+
+		SetPaletteItems(PlacementModeSubsystem->GetModeSettingsObject()->GetActivePaletteItems());
+
+		// Force a reset on the content browser mirroring if necessary.
+		const UAssetPlacementSettings* PlacementSettings = PlacementModeSubsystem->GetModeSettingsObject();
+		bIsMirroringContentBrowser = !PlacementSettings->bUseContentBrowserSelection;
+		SetupContentBrowserMirroring(PlacementSettings->bUseContentBrowserSelection);
+	}
+}
+
 void SAssetPlacementPalette::SetViewMode(EViewMode NewViewMode)
 {
 	if (ActiveViewMode != NewViewMode)
@@ -705,9 +811,6 @@ EActiveTimerReturnType SAssetPlacementPalette::UpdatePaletteItems(double InCurre
 		{
 			FilteredItems.Add(Item);
 		}
-
-		FPaletteItem NewItem;
-		NewItem.AssetData = Item->GetTypeUIInfo()->AssetData;
 	}
 
 	// Refresh the appropriate view
diff --git a/Engine/Source/Editor/AssetPlacementEdMode/Private/SAssetPlacementPalette.h b/Engine/Source/Editor/AssetPlacementEdMode/Private/SAssetPlacementPalette.h
--- a/Engine/Source/Editor/AssetPlacementEdMode/Private/SAssetPlacementPalette.h
+++ b/Engine/Source/Editor/AssetPlacementEdMode/Private/SAssetPlacementPalette.h
@@ -16,6 +16,7 @@
 #include "Widgets/Views/STileView.h"
 #include "Widgets/Views/STreeView.h"
 #include "Misc/TextFilter.h"
+#include "PlacementPaletteAsset.h"
 
 class FAssetThumbnailPool;
 class FAssetPlacementPaletteItemModel;
@@ -43,7 +44,6 @@ public:
 	};
 
 	SLATE_BEGIN_ARGS(SAssetPlacementPalette) {}
-		SLATE_ARGUMENT(TSharedPtr<IPropertyHandle>, PalettePropertyHandle)
 	SLATE_END_ARGS()
 
 	void Construct(const FArguments& InArgs);
@@ -92,9 +92,16 @@ private:
 
 	bool ShouldFilterAsset(const FAssetData& InAssetData);
 
+	void OnResetPaletteAssetClicked();
+	void OnSavePaletteAssetClicked();
+	EVisibility GetContentBrowserMirrorVisibility() const;
 	void OnContentBrowserMirrorButtonClicked(ECheckBoxState InState);
 	void OnContentBrowserSelectionChanged(const TArray<FAssetData>& NewSelectedAssets, bool bIsPrimaryBrowser);
 	void SetupContentBrowserMirroring(bool bInMirrorContentBrowser);
+	void OnSetPaletteAsset(const FAssetData& InAssetData);
+	void SetPaletteItems(TArrayView<const FPaletteItem> InPaletteItems);
+	void SetToLocalUserPalette();
+	FString GetPalettePath() const;
 
 	/** Sets the view mode of the palette */
 	void SetViewMode(EViewMode NewViewMode);
@@ -188,8 +195,6 @@ private:
 	/** Placement mesh details widget  */
 	TSharedPtr<IDetailsView> DetailsWidget;
 
-	TSharedPtr<IPropertyHandle> PalettePropertyHandle;
-
 	/** Placement items search box widget */
 	TSharedPtr<SSearchBox> SearchBoxPtr;
 
@@ -199,6 +204,8 @@ private:
 	/** Thumbnail pool for rendering mesh thumbnails */
 	TSharedPtr<FAssetThumbnailPool> ThumbnailPool;
 
+	FSoftObjectPath PalettePath;
+
 	bool bItemsNeedRebuild = true;
 	bool bShowFullTooltips = true;
 	bool bIsRebuildTimerRegistered = true;
diff --git a/Engine/Source/Editor/AssetPlacementEdMode/Private/Tools/PlacementBrushToolBase.cpp b/Engine/Source/Editor/AssetPlacementEdMode/Private/Tools/PlacementBrushToolBase.cpp
--- a/Engine/Source/Editor/AssetPlacementEdMode/Private/Tools/PlacementBrushToolBase.cpp
+++ b/Engine/Source/Editor/AssetPlacementEdMode/Private/Tools/PlacementBrushToolBase.cpp
@@ -28,7 +28,7 @@
 bool UPlacementToolBuilderBase::CanBuildTool(const FToolBuilderState& SceneState) const
 {	
 	TWeakObjectPtr<const UAssetPlacementSettings> PlacementSettings = GEditor->GetEditorSubsystem<UPlacementModeSubsystem>()->GetModeSettingsObject();
-	return PlacementSettings.IsValid() && PlacementSettings->PaletteItems.Num();
+	return PlacementSettings.IsValid() && PlacementSettings->GetActivePaletteItems().Num();
 }
 
 UInteractiveTool* UPlacementToolBuilderBase::BuildTool(const FToolBuilderState& SceneState) const
@@ -295,7 +295,7 @@ FTypedElementListRef UPlacementBrushToolBase::GetElementsInBrushRadius(const FIn
 		if (UInstancedStaticMeshComponent* ISMComponent = Cast<UInstancedStaticMeshComponent>(Hit.GetComponent()))
 		{
 			FTypedElementHandle StaticMeshObjectHandle = UEngineElementsLibrary::AcquireEditorObjectElementHandle(ISMComponent->GetStaticMesh());
-			if (!GEditor->GetEditorSubsystem<UPlacementModeSubsystem>()->DoesCurrentPaletteSupportElement(StaticMeshObjectHandle))
+			if (!GEditor->GetEditorSubsystem<UPlacementModeSubsystem>()->DoesActivePaletteSupportElement(StaticMeshObjectHandle))
 			{
 				continue;
 			}
@@ -313,7 +313,7 @@ FTypedElementListRef UPlacementBrushToolBase::GetElementsInBrushRadius(const FIn
 		else
 		{
 			FTypedElementHandle PrimarySelectionHandle = SelectionSet->GetSelectionElement(UEngineElementsLibrary::AcquireEditorComponentElementHandle(Hit.GetComponent()), ETypedElementSelectionMethod::Primary);
-			if (!GEditor->GetEditorSubsystem<UPlacementModeSubsystem>()->DoesCurrentPaletteSupportElement(PrimarySelectionHandle))
+			if (!GEditor->GetEditorSubsystem<UPlacementModeSubsystem>()->DoesActivePaletteSupportElement(PrimarySelectionHandle))
 			{
 				continue;
 			}
diff --git a/Engine/Source/Editor/AssetPlacementEdMode/Private/Tools/PlacementPlaceSingleTool.cpp b/Engine/Source/Editor/AssetPlacementEdMode/Private/Tools/PlacementPlaceSingleTool.cpp
--- a/Engine/Source/Editor/AssetPlacementEdMode/Private/Tools/PlacementPlaceSingleTool.cpp
+++ b/Engine/Source/Editor/AssetPlacementEdMode/Private/Tools/PlacementPlaceSingleTool.cpp
@@ -23,6 +23,7 @@
 #include "EditorModeManager.h"
 #include "Toolkits/IToolkitHost.h"
 #include "ContextObjectStore.h"
+#include "PlacementPaletteItem.h"
 
 constexpr TCHAR UPlacementModePlaceSingleTool::ToolName[];
 
@@ -228,23 +229,25 @@ FInputRayHit UPlacementModePlaceSingleTool::BeginHoverSequenceHitTest(const FInp
 
 void UPlacementModePlaceSingleTool::GeneratePlacementData(const FInputDeviceRay& DevicePos)
 {
-	const UAssetPlacementSettings* PlacementSettings = GEditor->GetEditorSubsystem<UPlacementModeSubsystem>()->GetModeSettingsObject();
-	if (PlacementSettings && PlacementSettings->PaletteItems.Num())
+	UPlacementModeSubsystem* PlacementSubsystem = GEditor->GetEditorSubsystem<UPlacementModeSubsystem>();
+	const UAssetPlacementSettings* PlacementSettings = PlacementSubsystem ? PlacementSubsystem->GetModeSettingsObject() : nullptr;
+	if (PlacementSettings)
 	{
-		int32 ItemIndex = FMath::RandHelper(PlacementSettings->PaletteItems.Num());
-		const TSharedPtr<FPaletteItem>& ItemToPlace = PlacementSettings->PaletteItems[ItemIndex];
-		if (!ItemToPlace)
+		TArrayView<const FPaletteItem> CurrentPaletteItems = PlacementSettings->GetActivePaletteItems();
+		if (CurrentPaletteItems.Num())
 		{
-			return;
+			int32 ItemIndex = FMath::RandHelper(CurrentPaletteItems.Num());
+			const FPaletteItem& ItemToPlace = CurrentPaletteItems[ItemIndex];
+			
+			FTransform TransformToUpdate(GenerateRandomRotation(PlacementSettings), LastBrushStamp.WorldPosition, GenerateRandomScale(PlacementSettings));
+
+			PlacementInfo = MakeUnique<FAssetPlacementInfo>();
+			PlacementInfo->AssetToPlace = ItemToPlace.AssetPath.TryLoad();
+			PlacementInfo->FactoryOverride = ItemToPlace.AssetFactoryInterface;
+			PlacementInfo->ItemGuid = ItemToPlace.ItemGuid;
+			PlacementInfo->FinalizedTransform = FinalizeTransform(TransformToUpdate, LastBrushStamp.WorldNormal, PlacementSettings);
+			PlacementInfo->PreferredLevel = GEditor->GetEditorWorldContext().World()->GetCurrentLevel();
 		}
-
-		FTransform TransformToUpdate(GenerateRandomRotation(PlacementSettings), LastBrushStamp.WorldPosition, GenerateRandomScale(PlacementSettings));
-
-		PlacementInfo = MakeUnique<FAssetPlacementInfo>();
-		PlacementInfo->AssetToPlace = ItemToPlace->AssetData;
-		PlacementInfo->FactoryOverride = ItemToPlace->AssetFactoryInterface;
-		PlacementInfo->FinalizedTransform = FinalizeTransform(TransformToUpdate, LastBrushStamp.WorldNormal, PlacementSettings);
-		PlacementInfo->PreferredLevel = GEditor->GetEditorWorldContext().World()->GetCurrentLevel();
 	}
 }
 
diff --git a/Engine/Source/Editor/AssetPlacementEdMode/Private/Tools/PlacementPlaceTool.cpp b/Engine/Source/Editor/AssetPlacementEdMode/Private/Tools/PlacementPlaceTool.cpp
--- a/Engine/Source/Editor/AssetPlacementEdMode/Private/Tools/PlacementPlaceTool.cpp
+++ b/Engine/Source/Editor/AssetPlacementEdMode/Private/Tools/PlacementPlaceTool.cpp
@@ -9,6 +9,7 @@
 #include "InstancedFoliage.h"
 #include "InteractiveToolManager.h"
 #include "Modes/PlacementModeSubsystem.h"
+#include "PlacementPaletteItem.h"
 
 constexpr TCHAR UPlacementModePlacementTool::ToolName[];
 
@@ -38,76 +39,78 @@ void UPlacementModePlacementTool::OnTick(float DeltaTime)
 		return;
 	}
 
-	TWeakObjectPtr<const UAssetPlacementSettings> PlacementSettings = GEditor->GetEditorSubsystem<UPlacementModeSubsystem>()->GetModeSettingsObject();
-	if (BrushProperties && PlacementSettings.IsValid() && PlacementSettings->PaletteItems.Num())
+	UPlacementModeSubsystem* PlacementModeSubsystem = GEditor->GetEditorSubsystem<UPlacementModeSubsystem>();
+	const UAssetPlacementSettings* PlacementSettings = PlacementModeSubsystem ? PlacementModeSubsystem->GetModeSettingsObject() : nullptr;
+	if (BrushProperties && PlacementSettings)
 	{
-		FVector BrushLocation = LastBrushStamp.WorldPosition;
-
-		// Assume a default density of 100 * whatever the user has selected as brush size.
-		float DefaultDensity = 100.0f * BrushProperties->BrushSize;
-		// This is the total set of instances disregarding parameters like slope, height or layer.
-		float DesiredInstanceCountFloat = DefaultDensity /** todo: individual item's Density*/ * BrushProperties->BrushStrength;
-		// Allow a single instance with a random chance, if the brush is smaller than the density
-		int32 DesiredInstanceCount = DesiredInstanceCountFloat > 1.f ? FMath::RoundToInt(DesiredInstanceCountFloat) : FMath::FRand() < DesiredInstanceCountFloat ? 1 : 0;
-
-		// Todo: save the instance hash per tile for the editor, so that paints don't continually add instances and surpass desired density (partition actor?)
-		FFoliageInstanceHash PotentialInstanceHash(7);
-		TArray<FVector> PotentialInstanceLocations;
-
-		TArray<FAssetPlacementInfo> DesiredInfoSpawnLocations;
-		DesiredInfoSpawnLocations.Reserve(DesiredInstanceCount);
-		PotentialInstanceLocations.Empty(DesiredInstanceCount);
-		for (int32 Index = 0; Index < DesiredInstanceCount; ++Index)
+		TArrayView<const FPaletteItem> PaletteItems = MakeArrayView(PlacementSettings->GetActivePaletteItems());
+		if (PaletteItems.Num())
 		{
-			FVector Start, End;
-			GetRandomVectorInBrush(Start, End);
-			FHitResult AdjustedHitResult;
-			FindHitResultWithStartAndEndTraceVectors(AdjustedHitResult, Start, End);
-			FVector SpawnLocation = AdjustedHitResult.ImpactPoint;
-			FVector SpawnNormal = AdjustedHitResult.ImpactNormal;
-
-			bool bValidInstance = true;
-			auto TempInstances = PotentialInstanceHash.GetInstancesOverlappingBox(FBox::BuildAABB(BrushLocation, FVector(BrushProperties->BrushRadius)));
-			for (int32 InstanceIndex : TempInstances)
+			FVector BrushLocation = LastBrushStamp.WorldPosition;
+
+			// Assume a default density of 100 * whatever the user has selected as brush size.
+			float DefaultDensity = 100.0f * BrushProperties->BrushSize;
+			// This is the total set of instances disregarding parameters like slope, height or layer.
+			float DesiredInstanceCountFloat = DefaultDensity /** todo: individual item's Density*/ * BrushProperties->BrushStrength;
+			// Allow a single instance with a random chance, if the brush is smaller than the density
+			int32 DesiredInstanceCount = DesiredInstanceCountFloat > 1.f ? FMath::RoundToInt(DesiredInstanceCountFloat) : FMath::FRand() < DesiredInstanceCountFloat ? 1 : 0;
+
+			// Todo: save the instance hash per tile for the editor, so that paints don't continually add instances and surpass desired density (partition actor?)
+			FFoliageInstanceHash PotentialInstanceHash(7);
+			TArray<FVector> PotentialInstanceLocations;
+
+			TArray<FAssetPlacementInfo> DesiredInfoSpawnLocations;
+			DesiredInfoSpawnLocations.Reserve(DesiredInstanceCount);
+			PotentialInstanceLocations.Empty(DesiredInstanceCount);
+			for (int32 Index = 0; Index < DesiredInstanceCount; ++Index)
 			{
-				if ((PotentialInstanceLocations[InstanceIndex] - BrushLocation).SizeSquared() < FMath::Square(BrushProperties->BrushRadius))
+				FVector Start, End;
+				GetRandomVectorInBrush(Start, End);
+				FHitResult AdjustedHitResult;
+				FindHitResultWithStartAndEndTraceVectors(AdjustedHitResult, Start, End);
+				FVector SpawnLocation = AdjustedHitResult.ImpactPoint;
+				FVector SpawnNormal = AdjustedHitResult.ImpactNormal;
+
+				bool bValidInstance = true;
+				auto TempInstances = PotentialInstanceHash.GetInstancesOverlappingBox(FBox::BuildAABB(BrushLocation, FVector(BrushProperties->BrushRadius)));
+				for (int32 InstanceIndex : TempInstances)
 				{
-					bValidInstance = false;
-					break;
+					if ((PotentialInstanceLocations[InstanceIndex] - BrushLocation).SizeSquared() < FMath::Square(BrushProperties->BrushRadius))
+					{
+						bValidInstance = false;
+						break;
+					}
 				}
-			}
 
-			if (!bValidInstance)
-			{
-				continue;
-			}
+				if (!bValidInstance)
+				{
+					continue;
+				}
 
-			int32 PotentialIdx = PotentialInstanceLocations.Add(SpawnLocation);
-			PotentialInstanceHash.InsertInstance(SpawnLocation, PotentialIdx);
+				int32 PotentialIdx = PotentialInstanceLocations.Add(SpawnLocation);
+				PotentialInstanceHash.InsertInstance(SpawnLocation, PotentialIdx);
 
-			int32 ItemIndex = FMath::RandHelper(PlacementSettings->PaletteItems.Num());
-			const TSharedPtr<FPaletteItem>& ItemToPlace = PlacementSettings->PaletteItems[ItemIndex];
-			if (!ItemToPlace)
-			{
-				return;
-			}
+				int32 ItemIndex = FMath::RandHelper(PaletteItems.Num());
+				const FPaletteItem& ItemToPlace = PaletteItems[ItemIndex];
 
-			FAssetPlacementInfo NewInfo;
-			NewInfo.AssetToPlace = ItemToPlace->AssetData;
-			NewInfo.FactoryOverride = ItemToPlace->AssetFactoryInterface;
-			NewInfo.PreferredLevel = GEditor->GetEditorWorldContext().World()->GetCurrentLevel();
-			NewInfo.FinalizedTransform = GenerateTransformFromHitLocationAndNormal(SpawnLocation, SpawnNormal);
+				FAssetPlacementInfo NewInfo;
+				NewInfo.AssetToPlace = ItemToPlace.AssetPath.TryLoad();
+				NewInfo.FactoryOverride = ItemToPlace.AssetFactoryInterface;
+				NewInfo.ItemGuid = ItemToPlace.ItemGuid;
+				NewInfo.PreferredLevel = GEditor->GetEditorWorldContext().World()->GetCurrentLevel();
+				NewInfo.FinalizedTransform = GenerateTransformFromHitLocationAndNormal(SpawnLocation, SpawnNormal);
 
-			DesiredInfoSpawnLocations.Emplace(NewInfo);
-		}
+				DesiredInfoSpawnLocations.Emplace(NewInfo);
+			}
 
-		UPlacementSubsystem* PlacementSubsystem = GEditor->GetEditorSubsystem<UPlacementSubsystem>();
-		if (PlacementSubsystem)
-		{
-			FPlacementOptions PlacementOptions;
-			PlacementOptions.bPreferBatchPlacement = true;
-			PlacementOptions.InstancedPlacementGridGuid = UPlacementSubsystem::GetUserGridGuid();
-			PlacementSubsystem->PlaceAssets(DesiredInfoSpawnLocations, PlacementOptions);
+			UPlacementSubsystem* PlacementSubsystem = GEditor->GetEditorSubsystem<UPlacementSubsystem>();
+			if (PlacementSubsystem)
+			{
+				FPlacementOptions PlacementOptions;
+				PlacementOptions.bPreferBatchPlacement = true;
+				PlacementOptions.InstancedPlacementGridGuid = UPlacementSubsystem::GetUserGridGuid();
+				PlacementSubsystem->PlaceAssets(DesiredInfoSpawnLocations, PlacementOptions);
+			}
 		}
 	}
 }
