Adding palette asset icon.

diff --git a/Engine/Source/Editor/AssetPlacementEdMode/Private/AssetPlacementEdModeStyle.cpp b/Engine/Source/Editor/AssetPlacementEdMode/Private/AssetPlacementEdModeStyle.cpp
--- a/Engine/Source/Editor/AssetPlacementEdMode/Private/AssetPlacementEdModeStyle.cpp
+++ b/Engine/Source/Editor/AssetPlacementEdMode/Private/AssetPlacementEdModeStyle.cpp
@@ -44,8 +44,8 @@ FAssetPlacementEdModeStyle::FAssetPlacementEdModeStyle()
 	SetCoreContentRoot(FPaths::EngineContentDir() / TEXT("Slate"));
 
 	Set("LevelEditor.AssetPlacementEdMode", new IMAGE_BRUSH_SVG("Starship/MainToolbar/AssetPlacementMode", Icon20x20));
-	Set("ClassIcon.PlacementPaletteAsset", new IMAGE_BRUSH_SVG("Starship/MainToolbar/AssetPlacementMode", FVector2D(16.0f, 16.0f)));
-	Set("ClassThumbnail.PlacementPaletteAsset", new IMAGE_BRUSH_SVG("Starship/MainToolbar/AssetPlacementMode", FVector2D(64.0f, 64.0f)));
+	Set("ClassIcon.PlacementPaletteAsset", new IMAGE_BRUSH_SVG("Starship/AssetIcons/AssetPalette", FVector2D(16.0f, 16.0f)));
+	Set("ClassThumbnail.PlacementPaletteAsset", new IMAGE_BRUSH_SVG("Starship/AssetIcons/AssetPalette", FVector2D(64.0f, 64.0f)));
 
 	Set("AssetPlacementEdMode.Select",                      new IMAGE_BRUSH("Icons/GeneralTools/Select_40x", Icon20x20));
 	Set("AssetPlacementEdMode.Select.Small",                new IMAGE_BRUSH("Icons/GeneralTools/Select_40x", Icon20x20));
