Fixing an issue where the input routing was not bubbling up to the palette view in the placement palette.

diff --git a/Engine/Source/Editor/AssetPlacementEdMode/Private/SAssetPlacementPalette.cpp b/Engine/Source/Editor/AssetPlacementEdMode/Private/SAssetPlacementPalette.cpp
--- a/Engine/Source/Editor/AssetPlacementEdMode/Private/SAssetPlacementPalette.cpp
+++ b/Engine/Source/Editor/AssetPlacementEdMode/Private/SAssetPlacementPalette.cpp
@@ -163,39 +163,28 @@ void SAssetPlacementPalette::Construct(const FArguments& InArgs)
 		]
 
 		+ SVerticalBox::Slot()
+		.HAlign(HAlign_Fill)
+		.VAlign(VAlign_Fill)
 		[
-			SNew(SOverlay)
-
-			+ SOverlay::Slot()
+			SNew(SAssetDropTarget)
+			.OnAreAssetsAcceptableForDrop(this, &SAssetPlacementPalette::OnAreAssetsValidForDrop)
+			.OnAssetsDropped(this, &SAssetPlacementPalette::HandlePlacementDropped)
+			.bSupportsMultiDrop(true)
 			[
-				SNew(SVerticalBox)
-
-				+ SVerticalBox::Slot()
-				.Padding(2.f, 0.f)
+				SNew(SOverlay)
+				+ SOverlay::Slot()
 				[
 					CreatePaletteViews()
 				]
-			]
-			
-			// Placement Drop Zone
-			+ SOverlay::Slot()
-			.HAlign(HAlign_Fill)
-			.VAlign(VAlign_Fill)
-			[
-				SNew(SAssetDropTarget)
-				.OnAreAssetsAcceptableForDrop(this, &SAssetPlacementPalette::OnAreAssetsValidForDrop)
-				.OnAssetsDropped(this, &SAssetPlacementPalette::HandlePlacementDropped)
-				.Visibility(this, &SAssetPlacementPalette::GetDropHintTextVisibility)
-				.bSupportsMultiDrop(true)
+
+				+ SOverlay::Slot()
+				.HAlign(HAlign_Center)
+				.VAlign(VAlign_Center)
 				[
-					SNew(SBox)
-					.HAlign(HAlign_Center)
-					.VAlign(VAlign_Center)
-					[
-						SNew(STextBlock)
-						.Text(LOCTEXT("Placement_AddPlacementMesh", "Drop Assets Here"))
-						.ShadowOffset(FVector2D(1.f, 1.f))
-					]
+					SNew(STextBlock)
+					.Visibility(this, &SAssetPlacementPalette::GetDropHintTextVisibility)
+					.Text(LOCTEXT("Placement_AddPlacementMesh", "Drop Assets Here"))
+					.ShadowOffset(FVector2D(1.f, 1.f))
 				]
 			]
 		]
