Add a call to modify when palette items are cleared from the active palette.

diff --git a/Engine/Source/Editor/AssetPlacementEdMode/Private/AssetPlacementSettings.cpp b/Engine/Source/Editor/AssetPlacementEdMode/Private/AssetPlacementSettings.cpp
--- a/Engine/Source/Editor/AssetPlacementEdMode/Private/AssetPlacementSettings.cpp
+++ b/Engine/Source/Editor/AssetPlacementEdMode/Private/AssetPlacementSettings.cpp
@@ -68,6 +68,7 @@ TArrayView<const FPaletteItem> UAssetPlacementSettings::GetActivePaletteItems()
 
 void UAssetPlacementSettings::ClearActivePaletteItems()
 {
+	ActivePalette->Modify();
 	ActivePalette->PaletteItems.Empty();
 }
 
