Fix a few bugs related to placement mode:
* palette asset not assigning grid guids when factoried or used to place
* palette asset is no longer a data asset
* user per project settings are saved under a good outer object
* mode having incorrect items in the palette on startup when a valid palette asset was last used

diff --git a/Engine/Source/Editor/AssetPlacementEdMode/Private/AssetPlacementEdMode.cpp b/Engine/Source/Editor/AssetPlacementEdMode/Private/AssetPlacementEdMode.cpp
--- a/Engine/Source/Editor/AssetPlacementEdMode/Private/AssetPlacementEdMode.cpp
+++ b/Engine/Source/Editor/AssetPlacementEdMode/Private/AssetPlacementEdMode.cpp
@@ -46,10 +46,13 @@ UAssetPlacementEdMode::~UAssetPlacementEdMode()
 
 void UAssetPlacementEdMode::Enter()
 {
-	// Set the settings object before we call Super::Enter, since we're using a shared one from the subsystem.
-	SettingsObject = const_cast<UAssetPlacementSettings*>(GEditor->GetEditorSubsystem<UPlacementModeSubsystem>()->GetModeSettingsObject());
-
 	Super::Enter();
+	if (UPlacementModeSubsystem* PlacementModeSubsystem = GEditor->GetEditorSubsystem<UPlacementModeSubsystem>())
+	{
+		SettingsObjectAsPlacementSettings = PlacementModeSubsystem->GetMutableModeSettingsObject();
+		SettingsObjectAsPlacementSettings->LoadSettings();
+		Toolkit->SetModeSettingsObject(SettingsObjectAsPlacementSettings.Get());
+	}
 
 	const FAssetPlacementEdModeCommands& PlacementModeCommands = FAssetPlacementEdModeCommands::Get();
 	RegisterTool(PlacementModeCommands.Select, UPlacementModeSelectTool::ToolName, NewObject<UPlacementModeSelectToolBuilder>(this));
@@ -73,16 +76,16 @@ void UAssetPlacementEdMode::Exit()
 {
 	Super::Exit();
 
+	SettingsObjectAsPlacementSettings->SaveSettings();
+	SettingsObjectAsPlacementSettings.Reset();
+
 	// Restore the selection to the original state after all the tools have shutdown in UEdMode::Exit()
 	// Since they can continue messing with selection states.
 	Owner->RestoreSelection(AssetPlacementEdModeID);
-
-	SettingsObjectAsPlacementSettings.Reset();
 }
 
 void UAssetPlacementEdMode::CreateToolkit()
 {
-	SettingsObjectAsPlacementSettings = Cast<UAssetPlacementSettings>(SettingsObject);
 	Toolkit = MakeShared<FAssetPlacementEdModeToolkit>();
 }
 
diff --git a/Engine/Source/Editor/AssetPlacementEdMode/Private/AssetPlacementEdModeModule.cpp b/Engine/Source/Editor/AssetPlacementEdMode/Private/AssetPlacementEdModeModule.cpp
--- a/Engine/Source/Editor/AssetPlacementEdMode/Private/AssetPlacementEdModeModule.cpp
+++ b/Engine/Source/Editor/AssetPlacementEdMode/Private/AssetPlacementEdModeModule.cpp
@@ -4,6 +4,8 @@
 #include "AssetPlacementEdModeCommands.h"
 #include "AssetPlacementEdModeStyle.h"
 #include "Modules/ModuleManager.h"
+#include "AssetTypeActions_PlacementPalette.h"
+#include "AssetToolsModule.h"
 
 #define LOCTEXT_NAMESPACE "AssetPlacementEdMode"
 
@@ -13,10 +15,20 @@ void FAssetPlacementEdMode::StartupModule()
 {
 	FAssetPlacementEdModeStyle::Get();
 	FAssetPlacementEdModeCommands::Register();
+
+	FAssetToolsModule& AssetToolsModule = FModuleManager::LoadModuleChecked<FAssetToolsModule>("AssetTools");
+	PaletteAssetActions = MakeShared<FAssetTypeActions_PlacementPalette>(EAssetTypeCategories::Type::Misc);
+	AssetToolsModule.Get().RegisterAssetTypeActions(PaletteAssetActions.ToSharedRef());
 }
 
 void FAssetPlacementEdMode::ShutdownModule()
 {
+	if (FAssetToolsModule* AssetToolsModule = FModuleManager::GetModulePtr<FAssetToolsModule>("AssetTools"))
+	{
+		AssetToolsModule->Get().UnregisterAssetTypeActions(PaletteAssetActions.ToSharedRef());
+	}
+	PaletteAssetActions.Reset();
+
 	FAssetPlacementEdModeCommands::Unregister();
 	FAssetPlacementEdModeStyle::Shutdown();
 }
diff --git a/Engine/Source/Editor/AssetPlacementEdMode/Private/AssetPlacementSettings.cpp b/Engine/Source/Editor/AssetPlacementEdMode/Private/AssetPlacementSettings.cpp
--- a/Engine/Source/Editor/AssetPlacementEdMode/Private/AssetPlacementSettings.cpp
+++ b/Engine/Source/Editor/AssetPlacementEdMode/Private/AssetPlacementSettings.cpp
@@ -66,18 +66,38 @@ TArrayView<const FPaletteItem> UAssetPlacementSettings::GetActivePaletteItems()
 	return ActivePalette ? MakeArrayView(ActivePalette->PaletteItems) : TArrayView<const FPaletteItem>();
 }
 
+FSoftObjectPath UAssetPlacementSettings::GetActivePalettePath() const
+{
+	return FSoftObjectPath(ActivePalette);
+}
+
+const FGuid UAssetPlacementSettings::GetActivePaletteGuid() const
+{
+	if (ActivePalette)
+	{
+		return ActivePalette->GridGuid;
+	}
+
+	return FGuid();
+}
+
 void UAssetPlacementSettings::ClearActivePaletteItems()
 {
 	ActivePalette->Modify();
 	ActivePalette->PaletteItems.Empty();
 }
 
-void UAssetPlacementSettings::Restore()
+void UAssetPlacementSettings::LoadSettings()
 {
 	LoadConfig();
 
 	UserPalette = NewObject<UPlacementPaletteAsset>(this);
-	UserPalette->LoadConfig();
+	if (!UserGridGuid.IsValid())
+	{
+		FGuid::Parse(FPlatformMisc::GetLoginId(), UserGridGuid);
+	}
+
+	UserPalette->GridGuid = UserGridGuid;
 
 	ActivePalette = Cast<UPlacementPaletteAsset>(LastActivePalettePath.TryLoad());
 	if (!ActivePalette)
@@ -98,7 +118,7 @@ void UAssetPlacementSettings::SaveActivePalette()
 	}
 }
 
-void UAssetPlacementSettings::Save()
+void UAssetPlacementSettings::SaveSettings()
 {
 	UPackageTools::SavePackagesForObjects(TArray<UObject*>({ActivePalette}));
 	UserPalette->SaveConfig();
diff --git a/Engine/Source/Editor/AssetPlacementEdMode/Private/AssetPlacementSettings.h b/Engine/Source/Editor/AssetPlacementEdMode/Private/AssetPlacementSettings.h
--- a/Engine/Source/Editor/AssetPlacementEdMode/Private/AssetPlacementSettings.h
+++ b/Engine/Source/Editor/AssetPlacementEdMode/Private/AssetPlacementSettings.h
@@ -124,10 +124,12 @@ public:
 	void SetPaletteAsset(UPlacementPaletteAsset* InPaletteAsset);
 	void AddItemToActivePalette(const FPaletteItem& InPaletteItem);
 	TArrayView<const FPaletteItem> GetActivePaletteItems() const;
+	const FGuid GetActivePaletteGuid() const;
 	void ClearActivePaletteItems();
 	void SaveActivePalette();
-	void Restore();
-	void Save();
+	FSoftObjectPath GetActivePalettePath() const;
+	void LoadSettings();
+	void SaveSettings();
 
 protected:
 
@@ -137,6 +139,9 @@ protected:
 	UPROPERTY(Transient)
 	TObjectPtr<UPlacementPaletteAsset> UserPalette;
 
+	UPROPERTY(config)
+	FGuid UserGridGuid;
+
 	UPROPERTY(config)
 	FSoftObjectPath LastActivePalettePath;
 };
diff --git a/Engine/Source/Editor/AssetPlacementEdMode/Private/AssetTypeActions_PlacementPalette.cpp b/Engine/Source/Editor/AssetPlacementEdMode/Private/AssetTypeActions_PlacementPalette.cpp
new file mode 100644
--- /dev/null
+++ b/Engine/Source/Editor/AssetPlacementEdMode/Private/AssetTypeActions_PlacementPalette.cpp
@@ -0,0 +1,30 @@
+// Copyright Epic Games, Inc. All Rights Reserved.
+
+#include "AssetTypeActions_PlacementPalette.h"
+
+#include "PlacementPaletteAsset.h"
+
+FAssetTypeActions_PlacementPalette::FAssetTypeActions_PlacementPalette(EAssetTypeCategories::Type InAssetCategoryBit)
+	: AssetCategoryBit(InAssetCategoryBit)
+{
+}
+
+FText FAssetTypeActions_PlacementPalette::GetName() const
+{
+	return NSLOCTEXT("AssetPlacementEdMode", "AssetTypeActions_PlacementPalette", "Placement Palette");
+}
+
+FColor FAssetTypeActions_PlacementPalette::GetTypeColor() const
+{
+	return FColor(222, 132, 255);
+}
+
+UClass* FAssetTypeActions_PlacementPalette::GetSupportedClass() const
+{
+	return UPlacementPaletteAsset::StaticClass();
+}
+
+uint32 FAssetTypeActions_PlacementPalette::GetCategories()
+{
+	return AssetCategoryBit;
+}
diff --git a/Engine/Source/Editor/AssetPlacementEdMode/Private/AssetTypeActions_PlacementPalette.h b/Engine/Source/Editor/AssetPlacementEdMode/Private/AssetTypeActions_PlacementPalette.h
new file mode 100644
--- /dev/null
+++ b/Engine/Source/Editor/AssetPlacementEdMode/Private/AssetTypeActions_PlacementPalette.h
@@ -0,0 +1,21 @@
+// Copyright Epic Games, Inc. All Rights Reserved.
+
+#pragma once
+
+#include "CoreMinimal.h"
+#include "AssetTypeActions_Base.h"
+
+class FAssetTypeActions_PlacementPalette : public FAssetTypeActions_Base
+{
+public:
+	FAssetTypeActions_PlacementPalette(EAssetTypeCategories::Type InAssetCategoryBit);
+
+	// IAssetTypeActions Implementation
+	virtual FText GetName() const override;
+	virtual FColor GetTypeColor() const override;
+	virtual UClass* GetSupportedClass() const override;
+	virtual uint32 GetCategories() override;
+
+private:
+	EAssetTypeCategories::Type AssetCategoryBit;
+};
diff --git a/Engine/Source/Editor/AssetPlacementEdMode/Private/Modes/PlacementModeSubsystem.cpp b/Engine/Source/Editor/AssetPlacementEdMode/Private/Modes/PlacementModeSubsystem.cpp
--- a/Engine/Source/Editor/AssetPlacementEdMode/Private/Modes/PlacementModeSubsystem.cpp
+++ b/Engine/Source/Editor/AssetPlacementEdMode/Private/Modes/PlacementModeSubsystem.cpp
@@ -16,12 +16,13 @@
 void UPlacementModeSubsystem::Initialize(FSubsystemCollectionBase& Collection)
 {
 	ModeSettings = NewObject<UAssetPlacementSettings>(this);
-	ModeSettings->Restore();
+	ModeSettings->LoadSettings();
+
+	FCoreDelegates::OnEnginePreExit.AddUObject(this, &UPlacementModeSubsystem::SaveSettings);
 }
 
 void UPlacementModeSubsystem::Deinitialize()
 {
-	ModeSettings->Save();
 	ModeSettings = nullptr;
 }
 
@@ -99,3 +100,11 @@ void UPlacementModeSubsystem::SetUseContentBrowserAsPalette(bool bInUseContentBr
 {
 	ModeSettings->bUseContentBrowserSelection = bInUseContentBrowser;
 }
+
+void UPlacementModeSubsystem::SaveSettings() const
+{
+	if (ModeSettings)
+	{
+		ModeSettings->SaveSettings();
+	}
+}
diff --git a/Engine/Source/Editor/AssetPlacementEdMode/Private/Modes/PlacementModeSubsystem.h b/Engine/Source/Editor/AssetPlacementEdMode/Private/Modes/PlacementModeSubsystem.h
--- a/Engine/Source/Editor/AssetPlacementEdMode/Private/Modes/PlacementModeSubsystem.h
+++ b/Engine/Source/Editor/AssetPlacementEdMode/Private/Modes/PlacementModeSubsystem.h
@@ -12,7 +12,7 @@ struct FTypedElementHandle;
 class UAssetPlacementSettings;
 class UPlacementPaletteAsset;
 
-UCLASS(Transient)
+UCLASS()
 class UPlacementModeSubsystem : public UEditorSubsystem
 {
 	GENERATED_BODY()
@@ -49,4 +49,6 @@ public:
 protected:
 	UPROPERTY()
 	TObjectPtr<UAssetPlacementSettings> ModeSettings;
+
+	void SaveSettings() const;
 };
diff --git a/Engine/Source/Editor/AssetPlacementEdMode/Private/PlacementPaletteAsset.cpp b/Engine/Source/Editor/AssetPlacementEdMode/Private/PlacementPaletteAsset.cpp
--- a/Engine/Source/Editor/AssetPlacementEdMode/Private/PlacementPaletteAsset.cpp
+++ b/Engine/Source/Editor/AssetPlacementEdMode/Private/PlacementPaletteAsset.cpp
@@ -12,10 +12,20 @@ UPlacementPaletteAssetFactory::UPlacementPaletteAssetFactory(const FObjectInitia
 
 UObject* UPlacementPaletteAssetFactory::FactoryCreateNew(UClass* InClass, UObject* InParent, FName InName, EObjectFlags Flags, UObject* Context, FFeedbackContext* Warn)
 {
-	return NewObject<UPlacementPaletteAsset>(InParent, InClass, InName, Flags, Context);
+	UPlacementPaletteAsset* PlacementPaletteAsset = NewObject<UPlacementPaletteAsset>(InParent, InClass, InName, Flags, Context);
+	if (PlacementPaletteAsset)
+	{
+		PlacementPaletteAsset->GridGuid = FGuid::NewGuid();
+	}
+	return PlacementPaletteAsset;
 }
 
 bool UPlacementPaletteAssetFactory::ShouldShowInNewMenu() const
 {
 	return true;
 }
+
+FText UPlacementPaletteAssetFactory::GetToolTip() const
+{
+	return NSLOCTEXT("AssetPlacementEdMode", "PlacementPaletteToolTip", "Placement palettes can be used in the placement mode to quickly place many different asset types in a level with preconfigured settings.");
+}
diff --git a/Engine/Source/Editor/AssetPlacementEdMode/Private/PlacementPaletteAsset.h b/Engine/Source/Editor/AssetPlacementEdMode/Private/PlacementPaletteAsset.h
--- a/Engine/Source/Editor/AssetPlacementEdMode/Private/PlacementPaletteAsset.h
+++ b/Engine/Source/Editor/AssetPlacementEdMode/Private/PlacementPaletteAsset.h
@@ -2,7 +2,7 @@
 
 #pragma once
 
-#include "Engine/DataAsset.h"
+#include "UObject/Object.h"
 #include "Factories/Factory.h"
 #include "PlacementPaletteItem.h"
 
@@ -10,14 +10,17 @@
 
 class UPlacementPaletteAsset;
 
-UCLASS(NotPlaceable, config = EditorPerProjectUserSettings)
-class UPlacementPaletteAsset : public UDataAsset
+UCLASS(NotPlaceable)
+class UPlacementPaletteAsset : public UObject
 {
 	GENERATED_BODY()
 
 public:
-	UPROPERTY(config, EditAnywhere, Category=Items)
+	UPROPERTY(EditAnywhere, Category=Items)
 	TArray<FPaletteItem> PaletteItems;
+
+	UPROPERTY()
+	FGuid GridGuid;
 };
 
 UCLASS(hideCategories=Object)
@@ -28,4 +31,5 @@ class UPlacementPaletteAssetFactory : public UFactory
 public:
 	virtual UObject* FactoryCreateNew(UClass* InClass, UObject* InParent, FName InName, EObjectFlags Flags, UObject* Context, FFeedbackContext* Warn) override;
 	virtual bool ShouldShowInNewMenu() const override;
+	virtual FText GetToolTip() const override;
 };
diff --git a/Engine/Source/Editor/AssetPlacementEdMode/Private/SAssetPlacementPalette.cpp b/Engine/Source/Editor/AssetPlacementEdMode/Private/SAssetPlacementPalette.cpp
--- a/Engine/Source/Editor/AssetPlacementEdMode/Private/SAssetPlacementPalette.cpp
+++ b/Engine/Source/Editor/AssetPlacementEdMode/Private/SAssetPlacementPalette.cpp
@@ -61,7 +61,15 @@ void SAssetPlacementPalette::Construct(const FArguments& InArgs)
 
 	const FText BlankText = FText::GetEmpty();
 
-	SetToLocalUserPalette();
+	// Load initial settings for the palette widget to use
+	if (UPlacementModeSubsystem* PlacementModeSubystem = GEditor->GetEditorSubsystem<UPlacementModeSubsystem>())
+	{
+		if (const UAssetPlacementSettings* PlacementSettings = PlacementModeSubystem->GetModeSettingsObject())
+		{
+			bIsMirroringContentBrowser = !PlacementSettings->bUseContentBrowserSelection;
+			OnSetPaletteAsset(PlacementSettings->GetActivePalettePath().TryLoad());
+		}
+	}
 
 	ChildSlot
 	[
@@ -113,14 +121,14 @@ void SAssetPlacementPalette::Construct(const FArguments& InArgs)
 							.VAlign(VAlign_Center)
 							.AutoWidth()
 							[
-								PropertyCustomizationHelpers::MakeClearButton(FSimpleDelegate::CreateSP(this, &SAssetPlacementPalette::OnResetPaletteAssetClicked), LOCTEXT("ResetPaletteAssetTooltip", "Clear the current palette asset, and use the user's local palette."))
+								PropertyCustomizationHelpers::MakeClearButton(FSimpleDelegate::CreateSP(this, &SAssetPlacementPalette::OnResetPaletteAssetClicked), LOCTEXT("ResetPaletteAssetTooltip", "Clear the current palette asset, and use the user's local palette."), TAttribute<bool>::CreateLambda([this]() { return PalettePath.IsValid(); }))
 							]
 							+ SHorizontalBox::Slot()
 							.VAlign(VAlign_Center)
 							.AutoWidth()
 							.Padding(2.0f, 0.0f)
 							[
-								PropertyCustomizationHelpers::MakeSaveButton(FSimpleDelegate::CreateSP(this, &SAssetPlacementPalette::OnSavePaletteAssetClicked), LOCTEXT("SaveAssetPaletteTooltip", "Save changes to the current palette asset"))
+								PropertyCustomizationHelpers::MakeSaveButton(FSimpleDelegate::CreateSP(this, &SAssetPlacementPalette::OnSavePaletteAssetClicked), LOCTEXT("SaveAssetPaletteTooltip", "Save changes to the current palette asset"), TAttribute<bool>::CreateLambda([this]() { return PalettePath.IsValid(); }))
 							]
 						]
 					]
@@ -184,7 +192,7 @@ void SAssetPlacementPalette::Construct(const FArguments& InArgs)
 					SNew(STextBlock)
 					.Visibility(this, &SAssetPlacementPalette::GetDropHintTextVisibility)
 					.Text(LOCTEXT("Placement_AddPlacementMesh", "Drop Assets Here"))
-					.ShadowOffset(FVector2D(1.f, 1.f))
+					.TextStyle(FAppStyle::Get(), "HintText")
 				]
 			]
 		]
@@ -372,8 +380,7 @@ bool SAssetPlacementPalette::ShouldFilterAsset(const FAssetData& InAssetData)
 
 void SAssetPlacementPalette::OnResetPaletteAssetClicked()
 {
-	PalettePath.Reset();
-	SetToLocalUserPalette();
+	OnSetPaletteAsset(FAssetData());
 	UpdatePalette(true);
 }
 
@@ -440,21 +447,19 @@ FString SAssetPlacementPalette::GetPalettePath() const
 void SAssetPlacementPalette::OnSetPaletteAsset(const FAssetData& InAssetData)
 {
 	FSoftObjectPath NewAssetPath = InAssetData.ToSoftObjectPath();
-	if (NewAssetPath.IsValid() && (NewAssetPath != PalettePath))
+	if ((NewAssetPath != PalettePath) || !InAssetData.IsValid())
 	{
-		if (UPlacementPaletteAsset* PaletteAsset = Cast<UPlacementPaletteAsset>(InAssetData.GetAsset()))
+		UPlacementPaletteAsset* PaletteAsset = Cast<UPlacementPaletteAsset>(InAssetData.GetAsset());
+		if (UPlacementModeSubsystem* PlacementModeSubsystem = GEditor->GetEditorSubsystem<UPlacementModeSubsystem>())
 		{
-			if (UPlacementModeSubsystem* PlacementModeSubsystem = GEditor->GetEditorSubsystem<UPlacementModeSubsystem>())
-			{
-				PlacementModeSubsystem->GetMutableModeSettingsObject()->SetPaletteAsset(PaletteAsset);
+			PlacementModeSubsystem->GetMutableModeSettingsObject()->SetPaletteAsset(PaletteAsset);
 
-				PalettePath = NewAssetPath;
+			PalettePath = NewAssetPath;
 
-				// Shutdown content browser mirroring if needed
-				SetupContentBrowserMirroring(false);
+			// Setup content browser mirroring if this is the user palette
+			SetupContentBrowserMirroring(PalettePath.IsValid() ? false : PlacementModeSubsystem->GetModeSettingsObject()->bUseContentBrowserSelection);
 
-				SetPaletteItems(PlacementModeSubsystem->GetModeSettingsObject()->GetActivePaletteItems());
-			}
+			SetPaletteItems(PlacementModeSubsystem->GetModeSettingsObject()->GetActivePaletteItems());
 		}
 	}
 }
@@ -483,21 +488,6 @@ void SAssetPlacementPalette::SetPaletteItems(TArrayView<const FPaletteItem> InPa
 	UpdatePalette(true);
 }
 
-void SAssetPlacementPalette::SetToLocalUserPalette()
-{
-	if (UPlacementModeSubsystem* PlacementModeSubsystem = GEditor->GetEditorSubsystem<UPlacementModeSubsystem>())
-	{
-		PlacementModeSubsystem->GetMutableModeSettingsObject()->SetPaletteAsset(nullptr);
-
-		SetPaletteItems(PlacementModeSubsystem->GetModeSettingsObject()->GetActivePaletteItems());
-
-		// Force a reset on the content browser mirroring if necessary.
-		const UAssetPlacementSettings* PlacementSettings = PlacementModeSubsystem->GetModeSettingsObject();
-		bIsMirroringContentBrowser = !PlacementSettings->bUseContentBrowserSelection;
-		SetupContentBrowserMirroring(PlacementSettings->bUseContentBrowserSelection);
-	}
-}
-
 void SAssetPlacementPalette::SetViewMode(EViewMode NewViewMode)
 {
 	if (ActiveViewMode != NewViewMode)
diff --git a/Engine/Source/Editor/AssetPlacementEdMode/Private/SAssetPlacementPalette.h b/Engine/Source/Editor/AssetPlacementEdMode/Private/SAssetPlacementPalette.h
--- a/Engine/Source/Editor/AssetPlacementEdMode/Private/SAssetPlacementPalette.h
+++ b/Engine/Source/Editor/AssetPlacementEdMode/Private/SAssetPlacementPalette.h
@@ -210,7 +210,7 @@ private:
 	bool bShowFullTooltips = true;
 	bool bIsRebuildTimerRegistered = true;
 	bool bIsRefreshTimerRegistered = true;
-	bool bIsMirroringContentBrowser = true;
+	bool bIsMirroringContentBrowser = false;
 	EViewMode ActiveViewMode;
 	EColumnSortMode::Type ActiveSortOrder = EColumnSortMode::Type::Ascending;
 
diff --git a/Engine/Source/Editor/AssetPlacementEdMode/Private/Tools/PlacementPlaceSingleTool.cpp b/Engine/Source/Editor/AssetPlacementEdMode/Private/Tools/PlacementPlaceSingleTool.cpp
--- a/Engine/Source/Editor/AssetPlacementEdMode/Private/Tools/PlacementPlaceSingleTool.cpp
+++ b/Engine/Source/Editor/AssetPlacementEdMode/Private/Tools/PlacementPlaceSingleTool.cpp
@@ -70,16 +70,18 @@ void UPlacementModePlaceSingleTool::OnClickPress(const FInputDeviceRay& PressPos
 {
 	Super::OnClickPress(PressPos);
 
-	// Place the Preview data if we managed to get to a valid handled click.
-	FPlacementOptions PlacementOptions;
-	PlacementOptions.bPreferBatchPlacement = true;
-	PlacementOptions.InstancedPlacementGridGuid = UPlacementSubsystem::GetUserGridGuid();
-
 	LastBrushStamp.Radius = 100.0f * LastBrushStampWorldToPixelScale;
-	
+
+	UPlacementModeSubsystem* PlacementModeSubsystem = GEditor->GetEditorSubsystem<UPlacementModeSubsystem>();
+	const UAssetPlacementSettings* PlacementSettings = PlacementModeSubsystem ? PlacementModeSubsystem->GetModeSettingsObject() : nullptr;
 	UPlacementSubsystem* PlacementSubsystem = GEditor->GetEditorSubsystem<UPlacementSubsystem>();
-	if (PlacementSubsystem)
+	if (PlacementSubsystem && PlacementSettings)
 	{
+		// Place the Preview data if we managed to get to a valid handled click.
+		FPlacementOptions PlacementOptions;
+		PlacementOptions.bPreferBatchPlacement = true;
+		PlacementOptions.InstancedPlacementGridGuid = PlacementSettings->GetActivePaletteGuid();
+
 		FAssetPlacementInfo FinalizedPlacementInfo = *PlacementInfo;
 		FinalizedPlacementInfo.FinalizedTransform = FinalizeTransform(
 			FTransform(PlacementInfo->FinalizedTransform.GetRotation(), LastBrushStamp.WorldPosition, PlacementInfo->FinalizedTransform.GetScale3D()),
diff --git a/Engine/Source/Editor/AssetPlacementEdMode/Private/Tools/PlacementPlaceTool.cpp b/Engine/Source/Editor/AssetPlacementEdMode/Private/Tools/PlacementPlaceTool.cpp
--- a/Engine/Source/Editor/AssetPlacementEdMode/Private/Tools/PlacementPlaceTool.cpp
+++ b/Engine/Source/Editor/AssetPlacementEdMode/Private/Tools/PlacementPlaceTool.cpp
@@ -108,7 +108,7 @@ void UPlacementModePlacementTool::OnTick(float DeltaTime)
 			{
 				FPlacementOptions PlacementOptions;
 				PlacementOptions.bPreferBatchPlacement = true;
-				PlacementOptions.InstancedPlacementGridGuid = UPlacementSubsystem::GetUserGridGuid();
+				PlacementOptions.InstancedPlacementGridGuid = PlacementSettings->GetActivePaletteGuid();
 				PlacementSubsystem->PlaceAssets(DesiredInfoSpawnLocations, PlacementOptions);
 			}
 		}
diff --git a/Engine/Source/Editor/AssetPlacementEdMode/Public/AssetPlacementEdModeModule.h b/Engine/Source/Editor/AssetPlacementEdMode/Public/AssetPlacementEdModeModule.h
--- a/Engine/Source/Editor/AssetPlacementEdMode/Public/AssetPlacementEdModeModule.h
+++ b/Engine/Source/Editor/AssetPlacementEdMode/Public/AssetPlacementEdModeModule.h
@@ -5,6 +5,8 @@
 #include "CoreMinimal.h"
 #include "Modules/ModuleInterface.h"
 
+class IAssetTypeActions;
+
 class FAssetPlacementEdMode : public IModuleInterface
 {
 public:
@@ -12,5 +14,5 @@ public:
 	virtual void ShutdownModule() override;
 
 protected:
-	void OnPostEngineInit();
+	TSharedPtr<IAssetTypeActions> PaletteAssetActions;
 };
