Remove duplicated viewport functionality from placement mode toolkit

diff --git a/Engine/Source/Editor/AssetPlacementEdMode/Private/AssetPlacementEdMode.cpp b/Engine/Source/Editor/AssetPlacementEdMode/Private/AssetPlacementEdMode.cpp
--- a/Engine/Source/Editor/AssetPlacementEdMode/Private/AssetPlacementEdMode.cpp
+++ b/Engine/Source/Editor/AssetPlacementEdMode/Private/AssetPlacementEdMode.cpp
@@ -103,14 +103,6 @@ void UAssetPlacementEdMode::BindCommands()
 {
 	const TSharedRef<FUICommandList>& CommandList = Toolkit->GetToolkitCommands();
 	const FAssetPlacementEdModeCommands& PlacementModeCommands = FAssetPlacementEdModeCommands::Get();
-
-	CommandList->MapAction(PlacementModeCommands.Deselect,
-		FExecuteAction::CreateUObject(this, &UAssetPlacementEdMode::ClearSelection),
-		FCanExecuteAction::CreateUObject(this, &UAssetPlacementEdMode::HasActiveSelection));
-
-	CommandList->MapAction(PlacementModeCommands.Delete,
-		FExecuteAction::CreateUObject(this, &UAssetPlacementEdMode::DeleteSelection),
-		FCanExecuteAction::CreateUObject(this, &UAssetPlacementEdMode::HasActiveSelection));
 }
 
 bool UAssetPlacementEdMode::IsSelectionAllowed(AActor* InActor, bool bInSelection) const
@@ -200,18 +192,6 @@ bool UAssetPlacementEdMode::IsEnabled()
 	return SMInstanceElementDataUtil::SMInstanceElementsEnabled() && GetDefault<ULevelEditorMiscSettings>()->bEnableAssetPlacementMode;
 }
 
-void UAssetPlacementEdMode::DeleteSelection()
-{
-	GetToolManager()->BeginUndoTransaction(LOCTEXT("PlacementDeleteAllSelected", "Delete Selected Assets"));
-
-	if (UTypedElementCommonActions* CommonActions = Owner->GetToolkitHost()->GetCommonActions())
-	{
-		CommonActions->DeleteSelectedElements(Owner->GetEditorSelectionSet(), GetWorld(), FTypedElementDeletionOptions());
-	}
-
-	GetToolManager()->EndUndoTransaction();
-}
-
 void UAssetPlacementEdMode::ClearSelection()
 {
 	GetToolManager()->BeginUndoTransaction(LOCTEXT("PlacementClearSelection", "Clear Selection"));
diff --git a/Engine/Source/Editor/AssetPlacementEdMode/Private/AssetPlacementEdModeCommands.cpp b/Engine/Source/Editor/AssetPlacementEdMode/Private/AssetPlacementEdModeCommands.cpp
--- a/Engine/Source/Editor/AssetPlacementEdMode/Private/AssetPlacementEdModeCommands.cpp
+++ b/Engine/Source/Editor/AssetPlacementEdMode/Private/AssetPlacementEdModeCommands.cpp
@@ -16,20 +16,16 @@ FAssetPlacementEdModeCommands::FAssetPlacementEdModeCommands()
 void FAssetPlacementEdModeCommands::RegisterCommands()
 {
 	TArray <TSharedPtr<FUICommandInfo>>& ToolCommands = Commands.FindOrAdd(NAME_Default);
-	UI_COMMAND(Select, "Select", "Select", EUserInterfaceActionType::ToggleButton, FInputChord());
+	UI_COMMAND(Select, "Select", "Select by clicking single assets matching the active palette.", EUserInterfaceActionType::ToggleButton, FInputChord());
 	ToolCommands.Add(Select);
-	UI_COMMAND(Deselect, "Deselect", "Clears the Current Selection.", EUserInterfaceActionType::ToggleButton, FInputChord());
-	ToolCommands.Add(Deselect);
-	UI_COMMAND(LassoSelect, "Lasso", "Select Assets with a Brush Tool.", EUserInterfaceActionType::ToggleButton, FInputChord());
+	UI_COMMAND(LassoSelect, "Lasso", "Selects asset by painting the area to select.", EUserInterfaceActionType::ToggleButton, FInputChord());
 	ToolCommands.Add(LassoSelect);
-	UI_COMMAND(Place, "Place", "Place Assets from the Palette.", EUserInterfaceActionType::ToggleButton, FInputChord());
+	UI_COMMAND(Place, "Paint", "Paint mutliple assets from the active palette.", EUserInterfaceActionType::ToggleButton, FInputChord());
 	ToolCommands.Add(Place);
-	UI_COMMAND(PlaceSingle, "Single", "Place a Single Asset from the Palette.", EUserInterfaceActionType::ToggleButton, FInputChord());
+	UI_COMMAND(PlaceSingle, "Stamp", "Stamp a random asset from the active palette.", EUserInterfaceActionType::ToggleButton, FInputChord());
 	ToolCommands.Add(PlaceSingle);
-	UI_COMMAND(Erase, "Erase", "Erases Assets Matching the Palette.", EUserInterfaceActionType::ToggleButton, FInputChord());
+	UI_COMMAND(Erase, "Erase", "Paint to erase assets matching the active palette.", EUserInterfaceActionType::ToggleButton, FInputChord());
 	ToolCommands.Add(Erase);
-	UI_COMMAND(Delete, "Delete", "Deletes Assets in the Current Selection.", EUserInterfaceActionType::ToggleButton, FInputChord());
-	ToolCommands.Add(Delete);
 }
 
 TMap<FName, TArray<TSharedPtr<FUICommandInfo>>> FAssetPlacementEdModeCommands::GetCommands()
diff --git a/Engine/Source/Editor/AssetPlacementEdMode/Private/AssetPlacementEdModeStyle.cpp b/Engine/Source/Editor/AssetPlacementEdMode/Private/AssetPlacementEdModeStyle.cpp
--- a/Engine/Source/Editor/AssetPlacementEdMode/Private/AssetPlacementEdModeStyle.cpp
+++ b/Engine/Source/Editor/AssetPlacementEdMode/Private/AssetPlacementEdModeStyle.cpp
@@ -49,10 +49,6 @@ FAssetPlacementEdModeStyle::FAssetPlacementEdModeStyle()
 
 	Set("AssetPlacementEdMode.Select",                      new IMAGE_BRUSH("Icons/GeneralTools/Select_40x", Icon20x20));
 	Set("AssetPlacementEdMode.Select.Small",                new IMAGE_BRUSH("Icons/GeneralTools/Select_40x", Icon20x20));
-	Set("AssetPlacementEdMode.SelectAll",                   new IMAGE_BRUSH("Icons/GeneralTools/SelectAll_40x", Icon20x20));
-	Set("AssetPlacementEdMode.SelectAll.Small",             new IMAGE_BRUSH("Icons/GeneralTools/SelectAll_40x", Icon20x20));
-	Set("AssetPlacementEdMode.Deselect",                    new IMAGE_BRUSH("Icons/GeneralTools/Deselect_40x", Icon20x20));
-	Set("AssetPlacementEdMode.Deselect.Small",              new IMAGE_BRUSH("Icons/GeneralTools/Deselect_40x", Icon20x20));
 	Set("AssetPlacementEdMode.SelectInvalid",               new IMAGE_BRUSH("Icons/GeneralTools/SelectInvalid_40x", Icon20x20));
 	Set("AssetPlacementEdMode.SelectInvalid.Small",         new IMAGE_BRUSH("Icons/GeneralTools/SelectInvalid_40x", Icon20x20));
 	Set("AssetPlacementEdMode.LassoSelect",                 new IMAGE_BRUSH("Icons/GeneralTools/Lasso_40x", Icon20x20));
@@ -65,8 +61,6 @@ FAssetPlacementEdModeStyle::FAssetPlacementEdModeStyle()
 	Set("AssetPlacementEdMode.ReapplySettings.Small",       new IMAGE_BRUSH("Icons/GeneralTools/Repaint_40x", Icon20x20));
 	Set("AssetPlacementEdMode.Fill",                        new IMAGE_BRUSH("Icons/GeneralTools/PaintBucket_40x", Icon20x20));
 	Set("AssetPlacementEdMode.Fill.Small",                  new IMAGE_BRUSH("Icons/GeneralTools/PaintBucket_40x", Icon20x20));
-	Set("AssetPlacementEdMode.Delete",                      new IMAGE_BRUSH("Icons/GeneralTools/Delete_40x", Icon20x20));
-	Set("AssetPlacementEdMode.Delete.Small",                new IMAGE_BRUSH("Icons/GeneralTools/Delete_40x", Icon20x20));
 	Set("AssetPlacementEdMode.Erase",                       new IMAGE_BRUSH("Icons/GeneralTools/Erase_40x", Icon20x20));
 	Set("AssetPlacementEdMode.Erase.Small",                 new IMAGE_BRUSH("Icons/GeneralTools/Erase_40x", Icon20x20));
 	Set("AssetPlacementEdMode.Filter",                      new IMAGE_BRUSH("Icons/GeneralTools/Filter_40x", Icon20x20));
diff --git a/Engine/Source/Editor/AssetPlacementEdMode/Public/AssetPlacementEdMode.h b/Engine/Source/Editor/AssetPlacementEdMode/Public/AssetPlacementEdMode.h
--- a/Engine/Source/Editor/AssetPlacementEdMode/Public/AssetPlacementEdMode.h
+++ b/Engine/Source/Editor/AssetPlacementEdMode/Public/AssetPlacementEdMode.h
@@ -45,13 +45,6 @@ public:
 	static bool IsEnabled();
 
 protected:
-	enum class ESelectMode
-	{
-		Select,
-		Deselect,
-	};
-
-	void DeleteSelection();
 	void ClearSelection();
 	bool HasAnyAssetsInPalette() const;
 	bool HasActiveSelection() const;
diff --git a/Engine/Source/Editor/AssetPlacementEdMode/Public/AssetPlacementEdModeCommands.h b/Engine/Source/Editor/AssetPlacementEdMode/Public/AssetPlacementEdModeCommands.h
--- a/Engine/Source/Editor/AssetPlacementEdMode/Public/AssetPlacementEdModeCommands.h
+++ b/Engine/Source/Editor/AssetPlacementEdMode/Public/AssetPlacementEdModeCommands.h
@@ -15,12 +15,10 @@ public:
 
 	TSharedPtr<FUICommandInfo> Select;
 	TSharedPtr<FUICommandInfo> SelectAll;
-	TSharedPtr<FUICommandInfo> Deselect;
 	TSharedPtr<FUICommandInfo> LassoSelect;
 	TSharedPtr<FUICommandInfo> Place;
 	TSharedPtr<FUICommandInfo> PlaceSingle;
 	TSharedPtr<FUICommandInfo> Erase;
-	TSharedPtr<FUICommandInfo> Delete;
 
 protected:
 	TMap<FName, TArray<TSharedPtr<FUICommandInfo>>> Commands;
