Fixing UI command text that was submitted unintentionally for placement mode single place tool.

diff --git a/Engine/Source/Editor/AssetPlacementEdMode/Private/AssetPlacementEdModeCommands.cpp b/Engine/Source/Editor/AssetPlacementEdMode/Private/AssetPlacementEdModeCommands.cpp
--- a/Engine/Source/Editor/AssetPlacementEdMode/Private/AssetPlacementEdModeCommands.cpp
+++ b/Engine/Source/Editor/AssetPlacementEdMode/Private/AssetPlacementEdModeCommands.cpp
@@ -22,7 +22,7 @@ void FAssetPlacementEdModeCommands::RegisterCommands()
 	ToolCommands.Add(LassoSelect);
 	UI_COMMAND(Place, "Paint", "Paint mutliple assets from the active palette.", EUserInterfaceActionType::ToggleButton, FInputChord());
 	ToolCommands.Add(Place);
-	UI_COMMAND(PlaceSingle, "Stamp", "Stamp a random asset from the active palette.", EUserInterfaceActionType::ToggleButton, FInputChord());
+	UI_COMMAND(PlaceSingle, "Single", "Place a single, random asset from the active palette.", EUserInterfaceActionType::ToggleButton, FInputChord());
 	ToolCommands.Add(PlaceSingle);
 	UI_COMMAND(Erase, "Erase", "Paint to erase assets matching the active palette.", EUserInterfaceActionType::ToggleButton, FInputChord());
 	ToolCommands.Add(Erase);
