Fix issues with single place tool holding onto element references beyond their lifetime in undo/redo stacks, and fix an issue where HISM were not rebuilding trees right after placement.

diff --git a/Engine/Source/Editor/AssetPlacementEdMode/Private/Tools/PlacementPlaceSingleTool.cpp b/Engine/Source/Editor/AssetPlacementEdMode/Private/Tools/PlacementPlaceSingleTool.cpp
--- a/Engine/Source/Editor/AssetPlacementEdMode/Private/Tools/PlacementPlaceSingleTool.cpp
+++ b/Engine/Source/Editor/AssetPlacementEdMode/Private/Tools/PlacementPlaceSingleTool.cpp
@@ -70,6 +70,7 @@ void UPlacementModePlaceSingleTool::OnClickPress(const FInputDeviceRay& PressPos
 {
 	Super::OnClickPress(PressPos);
 
+	DestroyPreviewElements();
 	LastBrushStamp.Radius = 100.0f * LastBrushStampWorldToPixelScale;
 
 	UPlacementModeSubsystem* PlacementModeSubsystem = GEditor->GetEditorSubsystem<UPlacementModeSubsystem>();
@@ -178,6 +179,7 @@ void UPlacementModePlaceSingleTool::OnClickRelease(const FInputDeviceRay& Releas
 
 	ShutdownBrushStampIndicator();
 	PlacementInfo.Reset();
+	PlacedElements.Empty();
 }
 
 void UPlacementModePlaceSingleTool::OnBeginHover(const FInputDeviceRay& DevicePos)
@@ -304,7 +306,21 @@ void UPlacementModePlaceSingleTool::UpdatePreviewElements(const FInputDeviceRay&
 
 void UPlacementModePlaceSingleTool::DestroyPreviewElements()
 {
-	NotifyMovementEnded(PreviewElements);
+	if (IAssetEditorContextInterface* AssetEditorContext = GetToolManager()->GetContextObjectStore()->FindContext<IAssetEditorContextInterface>())
+	{
+		if (UTypedElementSelectionSet* SelectionSet = AssetEditorContext->GetMutableSelectionSet())
+		{
+			for (const FTypedElementHandle& PreviewElement : PreviewElements)
+			{
+				if (TTypedElement<ITypedElementWorldInterface> WorldInterfaceElement = UTypedElementRegistry::GetInstance()->GetElement<ITypedElementWorldInterface>(PreviewElement))
+				{
+					WorldInterfaceElement.NotifyMovementEnded();
+					WorldInterfaceElement.DeleteElement(WorldInterfaceElement.GetOwnerWorld(), SelectionSet, FTypedElementDeletionOptions());
+				}
+			}
+		}
+	}
+
 	PreviewElements.Empty();
 }
 
@@ -339,7 +355,6 @@ void UPlacementModePlaceSingleTool::ExitTweakState(bool bClearSelectionSet)
 	}
 
 	bIsTweaking = false;
-	PlacedElements.Empty();
 }
 
 void UPlacementModePlaceSingleTool::UpdateElementTransforms(TArrayView<const FTypedElementHandle> InElements, const FTransform& InTransform, bool bLocalTransform)
@@ -358,9 +373,9 @@ void UPlacementModePlaceSingleTool::UpdateElementTransforms(TArrayView<const FTy
 void UPlacementModePlaceSingleTool::NotifyMovementStarted(TArrayView<const FTypedElementHandle> InElements)
 {
 	// Notify Movement started
-	for (const FTypedElementHandle& PreviewElement : PreviewElements)
+	for (const FTypedElementHandle& ElementToNotify : InElements)
 	{
-		if (TTypedElement<ITypedElementWorldInterface> WorldInterfaceElement = UTypedElementRegistry::GetInstance()->GetElement<ITypedElementWorldInterface>(PreviewElement))
+		if (TTypedElement<ITypedElementWorldInterface> WorldInterfaceElement = UTypedElementRegistry::GetInstance()->GetElement<ITypedElementWorldInterface>(ElementToNotify))
 		{
 			WorldInterfaceElement.NotifyMovementStarted();
 		}
@@ -369,24 +384,11 @@ void UPlacementModePlaceSingleTool::NotifyMovementStarted(TArrayView<const FType
 
 void UPlacementModePlaceSingleTool::NotifyMovementEnded(TArrayView<const FTypedElementHandle> InElements)
 {
-	IAssetEditorContextInterface* AssetEditorContext = GetToolManager()->GetContextObjectStore()->FindContext<IAssetEditorContextInterface>();
-	if (!AssetEditorContext)
-	{
-		return;
-	}
-
-	UTypedElementSelectionSet* SelectionSet = AssetEditorContext->GetMutableSelectionSet();
-	if (!SelectionSet)
-	{
-		return;
-	}
-
-	for (const FTypedElementHandle& PreviewElement : PreviewElements)
+	for (const FTypedElementHandle& ElementToNotify : InElements)
 	{
-		if (TTypedElement<ITypedElementWorldInterface> WorldInterfaceElement = UTypedElementRegistry::GetInstance()->GetElement<ITypedElementWorldInterface>(PreviewElement))
+		if (TTypedElement<ITypedElementWorldInterface> WorldInterfaceElement = UTypedElementRegistry::GetInstance()->GetElement<ITypedElementWorldInterface>(ElementToNotify))
 		{
 			WorldInterfaceElement.NotifyMovementEnded();
-			WorldInterfaceElement.DeleteElement(WorldInterfaceElement.GetOwnerWorld(), SelectionSet, FTypedElementDeletionOptions());
 		}
 	}
 }
diff --git a/Engine/Source/Editor/UnrealEd/Private/Factories/EditorStaticMeshFactory.cpp b/Engine/Source/Editor/UnrealEd/Private/Factories/EditorStaticMeshFactory.cpp
--- a/Engine/Source/Editor/UnrealEd/Private/Factories/EditorStaticMeshFactory.cpp
+++ b/Engine/Source/Editor/UnrealEd/Private/Factories/EditorStaticMeshFactory.cpp
@@ -79,6 +79,7 @@ TArray<FTypedElementHandle> UEditorStaticMeshFactory::PlaceAsset(const FAssetPla
 			FISMClientHandle ClientHandle = PlacedElementsActor->RegisterClient(ItemGuidToUse);
 			TSortedMap<int32, TArray<FTransform>> InstanceMap;
 			InstanceMap.Emplace(PlacedElementsActor->RegisterISMComponentDescriptor(ComponentDescriptor), TArray({ FTransform() }));
+			ModifiedPartitionActors.Add(PlacedElementsActor);
 
 			TArray<FSMInstanceId> PlacedInstances = PlacedElementsActor->AddISMInstance(ClientHandle, InPlacementInfo.FinalizedTransform, InstanceMap);
 			for (const FSMInstanceId PlacedInstanceID : PlacedInstances)
@@ -136,3 +137,16 @@ bool UEditorStaticMeshFactory::ShouldPlaceInstancedStaticMeshes(const FPlacement
 {
 	return !InPlacementOptions.bIsCreatingPreviewElements && InPlacementOptions.InstancedPlacementGridGuid.IsValid();
 }
+
+void UEditorStaticMeshFactory::EndPlacement(TArrayView<const FTypedElementHandle> InPlacedElements, const FPlacementOptions& InPlacementOptions)
+{
+	for (TWeakObjectPtr<AISMPartitionActor> ISMPartitionActor : ModifiedPartitionActors)
+	{
+		if (ISMPartitionActor.IsValid())
+		{
+			ISMPartitionActor->UpdateHISMTrees(true, false);
+		}
+	}
+
+	ModifiedPartitionActors.Empty();
+}
diff --git a/Engine/Source/Editor/UnrealEd/Private/Factories/EditorStaticMeshFactory.h b/Engine/Source/Editor/UnrealEd/Private/Factories/EditorStaticMeshFactory.h
--- a/Engine/Source/Editor/UnrealEd/Private/Factories/EditorStaticMeshFactory.h
+++ b/Engine/Source/Editor/UnrealEd/Private/Factories/EditorStaticMeshFactory.h
@@ -7,7 +7,7 @@
 
 #include "EditorStaticMeshFactory.generated.h"
 
-struct FFoliageInfo;
+class AISMPartitionActor;
 
 UCLASS(Transient)
 class UEditorStaticMeshFactory : public UActorFactoryStaticMesh
@@ -18,8 +18,10 @@ public:
 	// Begin IAssetFactoryInterface
 	virtual TArray<FTypedElementHandle> PlaceAsset(const FAssetPlacementInfo& InPlacementInfo, const FPlacementOptions& InPlacementOptions) override;
 	virtual FAssetData GetAssetDataFromElementHandle(const FTypedElementHandle& InHandle) override;
+	virtual void EndPlacement(TArrayView<const FTypedElementHandle> InPlacedElements, const FPlacementOptions& InPlacementOptions) override;
 	// End IAssetFactoryInterface
 
 protected:
 	bool ShouldPlaceInstancedStaticMeshes(const FPlacementOptions& InPlacementOptions) const;
+	TSet<TWeakObjectPtr<AISMPartitionActor>> ModifiedPartitionActors;
 };
\ No newline at end of file
