Add ability for placement palette items to have a custom settings object from the placement factory, and add a details panel to the mode for editing the settings object.

diff --git a/Engine/Source/Editor/AssetPlacementEdMode/Private/AssetPlacementEdMode.cpp b/Engine/Source/Editor/AssetPlacementEdMode/Private/AssetPlacementEdMode.cpp
--- a/Engine/Source/Editor/AssetPlacementEdMode/Private/AssetPlacementEdMode.cpp
+++ b/Engine/Source/Editor/AssetPlacementEdMode/Private/AssetPlacementEdMode.cpp
@@ -121,7 +121,7 @@ bool UAssetPlacementEdMode::IsSelectionAllowed(AActor* InActor, bool bInSelectio
 
 	// And we need to have a valid palette item.
 	FTypedElementHandle ActorHandle = UEngineElementsLibrary::AcquireEditorActorElementHandle(InActor);
-	return GEditor->GetEditorSubsystem<UPlacementModeSubsystem>()->DoesActivePaletteSupportElement(ActorHandle);
+	return SettingsObjectAsPlacementSettings->DoesActivePaletteSupportElement(ActorHandle);
 }
 
 void UAssetPlacementEdMode::OnToolStarted(UInteractiveToolManager* Manager, UInteractiveTool* Tool)
diff --git a/Engine/Source/Editor/AssetPlacementEdMode/Private/AssetPlacementEdModeToolkit.cpp b/Engine/Source/Editor/AssetPlacementEdMode/Private/AssetPlacementEdModeToolkit.cpp
--- a/Engine/Source/Editor/AssetPlacementEdMode/Private/AssetPlacementEdModeToolkit.cpp
+++ b/Engine/Source/Editor/AssetPlacementEdMode/Private/AssetPlacementEdModeToolkit.cpp
@@ -8,11 +8,23 @@
 #include "Widgets/Layout/SExpandableArea.h"
 #include "Widgets/Layout/SScrollBox.h"
 #include "Toolkits/AssetEditorModeUILayer.h"
+#include "PropertyEditorModule.h"
+#include "Modules/ModuleManager.h"
 
 #define LOCTEXT_NAMESPACE "AssetPlacementEdModeToolkit"
 
 FAssetPlacementEdModeToolkit::FAssetPlacementEdModeToolkit()
 {
+	// Create the details view for the details view tab, so that it can be shared across our custom tabs
+	{
+		FPropertyEditorModule& PropertyEditorModule = FModuleManager::GetModuleChecked<FPropertyEditorModule>("PropertyEditor");
+		FDetailsViewArgs DetailsViewArgs;
+		DetailsViewArgs.bAllowSearch = false;
+		DetailsViewArgs.NameAreaSettings = FDetailsViewArgs::ENameAreaSettings::HideNameArea;
+		DetailsViewArgs.bHideSelectionTip = true;
+
+		PaletteItemDetailsView = PropertyEditorModule.CreateDetailView(DetailsViewArgs);
+	}
 }
 
 void FAssetPlacementEdModeToolkit::GetToolPaletteNames(TArray<FName>& PaletteNames) const
@@ -39,20 +51,9 @@ void FAssetPlacementEdModeToolkit::InvokeUI()
 	{
 		TSharedPtr<FAssetEditorModeUILayer> ModeUILayerPtr = ModeUILayer.Pin();
 		AssetPaletteTab = ModeUILayerPtr->GetTabManager()->TryInvokeTab(FAssetEditorModeUILayer::BottomLeftTabID);
-	}
-}
 
-TSharedPtr<SWidget> FAssetPlacementEdModeToolkit::GetInlineContent() const
-{
-	return 	SNew(SScrollBox)
-		+ SScrollBox::Slot()
-		[
-			DetailsView.ToSharedRef()
-		]
-		+ SScrollBox::Slot()
-		[
-			ModeDetailsView.ToSharedRef()
-		];
+		PaletteItemDetailsViewTab = ModeUILayerPtr->GetTabManager()->TryInvokeTab(FAssetEditorModeUILayer::BottomRightTabID);
+	}
 }
 
 void FAssetPlacementEdModeToolkit::RequestModeUITabs()
@@ -61,11 +62,17 @@ void FAssetPlacementEdModeToolkit::RequestModeUITabs()
 	if (ModeUILayer.IsValid())
 	{
 		TSharedPtr<FAssetEditorModeUILayer> ModeUILayerPtr = ModeUILayer.Pin();
-		TSharedRef<FWorkspaceItem> MenuGroup = ModeUILayerPtr->GetModeMenuCategory().ToSharedRef();
 		AssetPaletteInfo.OnSpawnTab = FOnSpawnTab::CreateSP(SharedThis(this), &FAssetPlacementEdModeToolkit::CreateAssetPalette);
-		AssetPaletteInfo.TabLabel = LOCTEXT("AssetPaletteTab", "Asset Palette");
-		AssetPaletteInfo.TabTooltip = LOCTEXT("ModesToolboxTabTooltipText", "Open the  Modes tab, which contains the active editor mode's settings.");
+		AssetPaletteInfo.TabLabel = LOCTEXT("AssetPaletteTab", "Palette");
+		AssetPaletteInfo.TabTooltip = LOCTEXT("AssetPaletteTabTooltipText", "Open the asset palette tab, which contains placement mode's palette settings.");
+		AssetPaletteInfo.TabIcon = GetEditorModeIcon();
 		ModeUILayerPtr->SetModePanelInfo(FAssetEditorModeUILayer::BottomLeftTabID, AssetPaletteInfo);
+
+		PaletteItemDetailsViewInfo.OnSpawnTab = FOnSpawnTab::CreateSP(SharedThis(this), &FAssetPlacementEdModeToolkit::CreatePaletteItemDetailsView);
+		PaletteItemDetailsViewInfo.TabLabel = LOCTEXT("PaletteItemDetailsTab", "Palette Details");
+		PaletteItemDetailsViewInfo.TabTooltip = LOCTEXT("PaletteItemDetailsTabTooltipText", "Open the asset palette details tab, which allows customization of individual items in the active palette.");
+		PaletteItemDetailsViewInfo.TabIcon = GetEditorModeIcon();
+		ModeUILayerPtr->SetModePanelInfo(FAssetEditorModeUILayer::BottomRightTabID, PaletteItemDetailsViewInfo);
 	}
 }
 
@@ -73,22 +80,50 @@ void FAssetPlacementEdModeToolkit::RequestModeUITabs()
 
 TSharedRef<SDockTab> FAssetPlacementEdModeToolkit::CreateAssetPalette(const FSpawnTabArgs& Args)
 {
-	TSharedPtr<SDockTab> CreatedTab = SNew(SDockTab)
+	return SNew(SDockTab)
 		[
 			SNew(SScrollBox)
 			+ SScrollBox::Slot()
 			[
 				SNew(SAssetPlacementPalette)
+				.ItemDetailsView(PaletteItemDetailsView)
 			]
 		];
+}
+
+TSharedRef<SDockTab> FAssetPlacementEdModeToolkit::CreatePaletteItemDetailsView(const FSpawnTabArgs& Args)
+{
+	return SNew(SDockTab)
+	[
+		SNew(SScrollBox)
+		+SScrollBox::Slot()
+		[
+			SNew(SOverlay)
+			+SOverlay::Slot()
+			[
+				PaletteItemDetailsView.ToSharedRef()
+			]
 
-	const FSlateBrush* TabIcon = GetEditorModeIcon().GetSmallIcon();
-	if (CreatedTab)
+			+SOverlay::Slot()
+			.Padding(0.0f, 4.0f, 0.0f, 0.0f)
+			.HAlign(HAlign_Center)
+			[
+				SNew(STextBlock)
+				.Text(LOCTEXT("PaletteDetailsViewHintText", "Select an item from the palette to edit settings."))
+				.Visibility(this, &FAssetPlacementEdModeToolkit::GetDetailsViewHintTextVisibility)
+				.TextStyle(FAppStyle::Get(), "HintText")
+			]
+		]
+	];
+}
+
+EVisibility FAssetPlacementEdModeToolkit::GetDetailsViewHintTextVisibility() const
+{
+	if (PaletteItemDetailsView && (PaletteItemDetailsView->GetSelectedObjects().Num() == 0))
 	{
-		CreatedTab->SetTabIcon(TabIcon);
+		return EVisibility::HitTestInvisible;
 	}
-	AssetPaletteTab = CreatedTab;
-	return CreatedTab.ToSharedRef();
+	return EVisibility::Collapsed;
 }
 
 #undef LOCTEXT_NAMESPACE
diff --git a/Engine/Source/Editor/AssetPlacementEdMode/Private/AssetPlacementPaletteItem.cpp b/Engine/Source/Editor/AssetPlacementEdMode/Private/AssetPlacementPaletteItem.cpp
--- a/Engine/Source/Editor/AssetPlacementEdMode/Private/AssetPlacementPaletteItem.cpp
+++ b/Engine/Source/Editor/AssetPlacementEdMode/Private/AssetPlacementPaletteItem.cpp
@@ -20,8 +20,15 @@
 
 #define LOCTEXT_NAMESPACE "AssetPlacementMode"
 
-FAssetPlacementPaletteItemModel::FAssetPlacementPaletteItemModel(const FAssetData& InAssetData, TSharedRef<SAssetPlacementPalette> InParentPalette, TSharedPtr<FAssetThumbnailPool> InThumbnailPool)
-	: TypeInfo(MakeShared<FPaletteItemUIInfo>(InAssetData))
+FPaletteItemUIInfo::FPaletteItemUIInfo(const FAssetData& InAssetData, const FPaletteItem& InPaletteItem)
+	: AssetData(InAssetData)
+	, SettingsObject(InPaletteItem.SettingsObject.Get())
+{
+	check(AssetData.IsValid());
+}
+
+FAssetPlacementPaletteItemModel::FAssetPlacementPaletteItemModel(const FAssetData& InAssetData, const FPaletteItem& InPaletteItem, TSharedRef<SAssetPlacementPalette> InParentPalette, TSharedPtr<FAssetThumbnailPool> InThumbnailPool)
+	: TypeInfo(MakeShared<FPaletteItemUIInfo>(InAssetData, InPaletteItem))
 	, AssetPalette(InParentPalette)
 {
 	check(TypeInfo);
@@ -214,7 +221,7 @@ void SAssetPlacementPaletteItemTile::Construct(const FArguments& InArgs, TShared
 
 	STableRow<FAssetPlacementUIInfoPtr>::Construct(
 		STableRow<FAssetPlacementUIInfoPtr>::FArguments()
-		.Style(FEditorStyle::Get(), "ContentBrowser.AssetListView.TableRow")
+		.Style(FEditorStyle::Get(), "ContentBrowser.AssetListView.ColumnListTableRow")
 		.Padding(1.f)
 		.Content()
 		[
diff --git a/Engine/Source/Editor/AssetPlacementEdMode/Private/AssetPlacementPaletteItem.h b/Engine/Source/Editor/AssetPlacementEdMode/Private/AssetPlacementPaletteItem.h
--- a/Engine/Source/Editor/AssetPlacementEdMode/Private/AssetPlacementPaletteItem.h
+++ b/Engine/Source/Editor/AssetPlacementEdMode/Private/AssetPlacementPaletteItem.h
@@ -24,13 +24,10 @@ namespace PlacementPaletteConstants
 
 struct FPaletteItemUIInfo
 {
-	FPaletteItemUIInfo(const FAssetData& InAssetData)
-		: AssetData(InAssetData)
-	{
-		check(AssetData.IsValid());
-	}
+	FPaletteItemUIInfo(const FAssetData& InAssetData, const FPaletteItem& InPaletteItem);
 
 	const FAssetData AssetData;
+	TWeakObjectPtr<UEditorFactorySettingsObject> SettingsObject;
 };
 
 typedef TSharedPtr<FPaletteItemUIInfo> FAssetPlacementUIInfoPtr;
@@ -43,7 +40,7 @@ namespace AssetPlacementPaletteTreeColumns
 class FAssetPlacementPaletteItemModel : public TSharedFromThis<FAssetPlacementPaletteItemModel>
 {
 public:
-	FAssetPlacementPaletteItemModel(const FAssetData& InAssetData, TSharedRef<class SAssetPlacementPalette> InParentPalette, TSharedPtr<class FAssetThumbnailPool> InThumbnailPool);
+	FAssetPlacementPaletteItemModel(const FAssetData& InAssetData, const FPaletteItem& InPaletteItem, TSharedRef<class SAssetPlacementPalette> InParentPalette, TSharedPtr<class FAssetThumbnailPool> InThumbnailPool);
 
 	FAssetPlacementUIInfoPtr GetTypeUIInfo() const;
 
diff --git a/Engine/Source/Editor/AssetPlacementEdMode/Private/AssetPlacementSettings.cpp b/Engine/Source/Editor/AssetPlacementEdMode/Private/AssetPlacementSettings.cpp
--- a/Engine/Source/Editor/AssetPlacementEdMode/Private/AssetPlacementSettings.cpp
+++ b/Engine/Source/Editor/AssetPlacementEdMode/Private/AssetPlacementSettings.cpp
@@ -1,9 +1,16 @@
 // Copyright Epic Games, Inc. All Rights Reserved.
 
 #include "AssetPlacementSettings.h"
-#include "PlacementPaletteAsset.h"
+
+#include "Editor.h"
 #include "PackageTools.h"
+#include "PlacementPaletteAsset.h"
 #include "PlacementPaletteItem.h"
+#include "Factories/AssetFactoryInterface.h"
+#include "Subsystems/PlacementSubsystem.h"
+#include "Elements/Framework/TypedElementHandle.h"
+#include "Elements/Framework/TypedElementRegistry.h"
+#include "Elements/Interfaces/TypedElementAssetDataInterface.h"
 
 bool UAssetPlacementSettings::CanEditChange(const FProperty* InProperty) const
 {
@@ -43,27 +50,52 @@ bool UAssetPlacementSettings::CanEditChange(const FProperty* InProperty) const
 
 void UAssetPlacementSettings::SetPaletteAsset(UPlacementPaletteAsset* InPaletteAsset)
 {
-	if (!InPaletteAsset)
+	ActivePalette = InPaletteAsset;
+	LastActivePalettePath = FSoftObjectPath(InPaletteAsset);
+}
+
+FPaletteItem UAssetPlacementSettings::AddItemToActivePalette(const FAssetData& InAssetData)
+{
+	FPaletteItem NewPaletteItem;
+	if (!InAssetData.IsValid())
 	{
-		ActivePalette = UserPalette;
+		return NewPaletteItem;
 	}
-	else
+
+	if (InAssetData.GetClass()->HasAnyClassFlags(CLASS_Abstract | CLASS_Deprecated | CLASS_NewerVersionExists | CLASS_NotPlaceable))
 	{
-		ActivePalette = InPaletteAsset;
+		return NewPaletteItem;
 	}
 
-	LastActivePalettePath = FSoftObjectPath(InPaletteAsset);
-}
+	if (!GetActivePaletteItems().FindByPredicate([InAssetData](const FPaletteItem& ItemIter) { return (ItemIter.AssetPath == InAssetData.ToSoftObjectPath()); }))
+	{
+		if (UPlacementSubsystem* PlacementSubystem = GEditor->GetEditorSubsystem<UPlacementSubsystem>())
+		{
+			if (TScriptInterface<IAssetFactoryInterface> AssetFactory = PlacementSubystem->FindAssetFactoryFromAssetData(InAssetData))
+			{
+				NewPaletteItem.ItemGuid = FGuid::NewGuid();
+				NewPaletteItem.AssetPath = InAssetData.ToSoftObjectPath();
+				NewPaletteItem.AssetFactoryInterface = AssetFactory;
+				FPlacementOptions PlacementOptions;
+				PlacementOptions.InstancedPlacementGridGuid = GetActivePaletteGuid();
+				NewPaletteItem.SettingsObject = AssetFactory->FactorySettingsObjectForPlacement(InAssetData, PlacementOptions);
+				GetMutableActivePalette()->Modify();
+				if (NewPaletteItem.SettingsObject)
+				{
+					// Set the outer of the new object to the current palette
+					NewPaletteItem.SettingsObject->Rename(nullptr, GetMutableActivePalette());
+				}
+				GetMutableActivePalette()->PaletteItems.Add(NewPaletteItem);
+			}
+		}
+	}
 
-void UAssetPlacementSettings::AddItemToActivePalette(const FPaletteItem& InPaletteItem)
-{
-	ActivePalette->Modify();
-	ActivePalette->PaletteItems.Add(InPaletteItem);
+	return NewPaletteItem;
 }
 
 TArrayView<const FPaletteItem> UAssetPlacementSettings::GetActivePaletteItems() const
 {
-	return ActivePalette ? MakeArrayView(ActivePalette->PaletteItems) : TArrayView<const FPaletteItem>();
+	return MakeArrayView(GetActivePalette()->PaletteItems);
 }
 
 FSoftObjectPath UAssetPlacementSettings::GetActivePalettePath() const
@@ -73,25 +105,24 @@ FSoftObjectPath UAssetPlacementSettings::GetActivePalettePath() const
 
 const FGuid UAssetPlacementSettings::GetActivePaletteGuid() const
 {
-	if (ActivePalette)
-	{
-		return ActivePalette->GridGuid;
-	}
-
-	return FGuid();
+	return GetActivePalette()->GridGuid;
 }
 
 void UAssetPlacementSettings::ClearActivePaletteItems()
 {
-	ActivePalette->Modify();
-	ActivePalette->PaletteItems.Empty();
+	GetMutableActivePalette()->Modify();
+	GetMutableActivePalette()->PaletteItems.Empty();
 }
 
 void UAssetPlacementSettings::LoadSettings()
 {
 	LoadConfig();
 
-	UserPalette = NewObject<UPlacementPaletteAsset>(this);
+	if (!UserPalette)
+	{
+		UserPalette = NewObject<UPlacementPaletteAsset>(this);
+	}
+
 	if (!UserGridGuid.IsValid())
 	{
 		FGuid::Parse(FPlatformMisc::GetLoginId(), UserGridGuid);
@@ -100,27 +131,65 @@ void UAssetPlacementSettings::LoadSettings()
 	UserPalette->GridGuid = UserGridGuid;
 
 	ActivePalette = Cast<UPlacementPaletteAsset>(LastActivePalettePath.TryLoad());
-	if (!ActivePalette)
-	{
-		ActivePalette = UserPalette;
-	}
+}
+
+void UAssetPlacementSettings::SaveSettings()
+{
+	SaveConfig();
 }
 
 void UAssetPlacementSettings::SaveActivePalette()
 {
-	if (ActivePalette != UserPalette)
+	if (ActivePalette)
 	{
 		UPackageTools::SavePackagesForObjects(TArray<UObject*>({ ActivePalette }));
 	}
-	else
+}
+
+bool UAssetPlacementSettings::DoesActivePaletteSupportElement(const FTypedElementHandle& InElementToCheck) const
+{	
+	if (TTypedElement<ITypedElementAssetDataInterface> AssetDataInterface = UTypedElementRegistry::GetInstance()->GetElement<ITypedElementAssetDataInterface>(InElementToCheck))
 	{
-		UserPalette->SaveConfig();
+		TArray<FAssetData> ReferencedAssetDatas = AssetDataInterface.GetAllReferencedAssetDatas();
+		for (const FPaletteItem& Item : GetActivePaletteItems())
+		{
+			if (ReferencedAssetDatas.FindByPredicate([&Item](const FAssetData& ReferencedAssetData) { return (ReferencedAssetData.ToSoftObjectPath() == Item.AssetPath); }))
+			{
+				return true;
+			}
+
+			// The current implementation of the asset data interface for actors requires that individual actors report on assets contained within their components.
+			// Not all actors do this reliably, so additionally check the supplied factory for a match. 
+			if (!Item.AssetFactoryInterface)
+			{
+				continue;
+			}
+
+			FAssetData FoundAssetDataFromFactory = Item.AssetFactoryInterface->GetAssetDataFromElementHandle(InElementToCheck);
+			if (FoundAssetDataFromFactory.ToSoftObjectPath() == Item.AssetPath)
+			{
+				return true;
+			}
+		}
 	}
+
+	return false;
 }
 
-void UAssetPlacementSettings::SaveSettings()
+const UPlacementPaletteAsset* UAssetPlacementSettings::GetActivePalette() const
 {
-	UPackageTools::SavePackagesForObjects(TArray<UObject*>({ActivePalette}));
-	UserPalette->SaveConfig();
-	SaveConfig();
+	if (ActivePalette)
+	{
+		return ActivePalette;
+	}
+	return UserPalette;
+}
+
+UPlacementPaletteAsset* UAssetPlacementSettings::GetMutableActivePalette()
+{
+	if (ActivePalette)
+	{
+		return ActivePalette;
+	}
+	return UserPalette;
 }
diff --git a/Engine/Source/Editor/AssetPlacementEdMode/Private/AssetPlacementSettings.h b/Engine/Source/Editor/AssetPlacementEdMode/Private/AssetPlacementSettings.h
--- a/Engine/Source/Editor/AssetPlacementEdMode/Private/AssetPlacementSettings.h
+++ b/Engine/Source/Editor/AssetPlacementEdMode/Private/AssetPlacementSettings.h
@@ -4,6 +4,7 @@
 
 #include "CoreMinimal.h"
 #include "FoliageType.h"
+#include "PlacementPaletteItem.h"
 
 #include "AssetPlacementSettings.generated.h"
 
@@ -122,7 +123,7 @@ public:
 	virtual bool CanEditChange(const FProperty* InProperty) const override;
 
 	void SetPaletteAsset(UPlacementPaletteAsset* InPaletteAsset);
-	void AddItemToActivePalette(const FPaletteItem& InPaletteItem);
+	FPaletteItem AddItemToActivePalette(const FAssetData& InAssetData);
 	TArrayView<const FPaletteItem> GetActivePaletteItems() const;
 	const FGuid GetActivePaletteGuid() const;
 	void ClearActivePaletteItems();
@@ -131,7 +132,16 @@ public:
 	void LoadSettings();
 	void SaveSettings();
 
+	/**
+	 * Verifies if the given element handle is supported by the current mode settings' palette.
+	 *
+	 * @returns true if the element can be placed by the mode.
+	 */
+	bool DoesActivePaletteSupportElement(const FTypedElementHandle& InElementToCheck) const;
+
 protected:
+	const UPlacementPaletteAsset* GetActivePalette() const;
+	UPlacementPaletteAsset* GetMutableActivePalette();
 
 	UPROPERTY(Transient)
 	TObjectPtr<UPlacementPaletteAsset> ActivePalette;
diff --git a/Engine/Source/Editor/AssetPlacementEdMode/Private/Modes/PlacementModeSubsystem.cpp b/Engine/Source/Editor/AssetPlacementEdMode/Private/Modes/PlacementModeSubsystem.cpp
--- a/Engine/Source/Editor/AssetPlacementEdMode/Private/Modes/PlacementModeSubsystem.cpp
+++ b/Engine/Source/Editor/AssetPlacementEdMode/Private/Modes/PlacementModeSubsystem.cpp
@@ -3,14 +3,8 @@
 #include "Modes/PlacementModeSubsystem.h"
 
 #include "AssetPlacementSettings.h"
-#include "Elements/Framework/TypedElementHandle.h"
-#include "Elements/Framework/TypedElementRegistry.h"
-#include "Elements/Interfaces/TypedElementAssetDataInterface.h"
-#include "Factories/AssetFactoryInterface.h"
-#include "PlacementPaletteAsset.h"
-#include "PlacementPaletteItem.h"
 
-#include "Subsystems/PlacementSubsystem.h"
+
 #include "Editor.h"
 
 void UPlacementModeSubsystem::Initialize(FSubsystemCollectionBase& Collection)
@@ -36,71 +30,6 @@ UAssetPlacementSettings* UPlacementModeSubsystem::GetMutableModeSettingsObject()
 	return ModeSettings;
 }
 
-bool UPlacementModeSubsystem::DoesActivePaletteSupportElement(const FTypedElementHandle& InElementToCheck) const
-{
-	if (TTypedElement<ITypedElementAssetDataInterface> AssetDataInterface = UTypedElementRegistry::GetInstance()->GetElement<ITypedElementAssetDataInterface>(InElementToCheck))
-	{
-		TArray<FAssetData> ReferencedAssetDatas = AssetDataInterface.GetAllReferencedAssetDatas();
-		for (const FPaletteItem& Item : ModeSettings->GetActivePaletteItems())
-		{
-			if (ReferencedAssetDatas.FindByPredicate([&Item](const FAssetData& ReferencedAssetData){ return (ReferencedAssetData.ToSoftObjectPath() == Item.AssetPath); }))
-			{
-				return true;
-			}
-
-			// The current implementation of the asset data interface for actors requires that individual actors report on assets contained within their components.
-			// Not all actors do this reliably, so additionally check the supplied factory for a match. 
-			if (!Item.AssetFactoryInterface)
-			{
-				continue;
-			}
-
-			FAssetData FoundAssetDataFromFactory = Item.AssetFactoryInterface->GetAssetDataFromElementHandle(InElementToCheck);
-			if (FoundAssetDataFromFactory.ToSoftObjectPath() == Item.AssetPath)
-			{
-				return true;
-			}
-		}
-	}
-
-	return false;
-}
-
-FPaletteItem UPlacementModeSubsystem::CreatePaletteItem(const FAssetData& InAssetData)
-{
-	FPaletteItem NewPaletteItem;
-	if (!InAssetData.IsValid())
-	{
-		return NewPaletteItem;
-	}
-
-	if (InAssetData.GetClass()->HasAnyClassFlags(CLASS_Abstract | CLASS_Deprecated | CLASS_NewerVersionExists | CLASS_NotPlaceable))
-	{
-		return NewPaletteItem;
-	}
-
-	if (!ModeSettings->GetActivePaletteItems().FindByPredicate([InAssetData](const FPaletteItem& ItemIter) { return (ItemIter.AssetPath == InAssetData.ToSoftObjectPath()); }))
-	{
-		if (UPlacementSubsystem* PlacementSubystem = GEditor->GetEditorSubsystem<UPlacementSubsystem>())
-		{
-			if (TScriptInterface<IAssetFactoryInterface> AssetFactory = PlacementSubystem->FindAssetFactoryFromAssetData(InAssetData))
-			{
-				NewPaletteItem.ItemGuid = FGuid::NewGuid();
-				NewPaletteItem.AssetPath = InAssetData.ToSoftObjectPath();
-				NewPaletteItem.AssetFactoryInterface = AssetFactory;
-				ModeSettings->AddItemToActivePalette(NewPaletteItem);
-			}
-		}
-	}
-
-	return NewPaletteItem;
-}
-
-void UPlacementModeSubsystem::SetUseContentBrowserAsPalette(bool bInUseContentBrowser)
-{
-	ModeSettings->bUseContentBrowserSelection = bInUseContentBrowser;
-}
-
 void UPlacementModeSubsystem::SaveSettings() const
 {
 	if (ModeSettings)
diff --git a/Engine/Source/Editor/AssetPlacementEdMode/Private/Modes/PlacementModeSubsystem.h b/Engine/Source/Editor/AssetPlacementEdMode/Private/Modes/PlacementModeSubsystem.h
--- a/Engine/Source/Editor/AssetPlacementEdMode/Private/Modes/PlacementModeSubsystem.h
+++ b/Engine/Source/Editor/AssetPlacementEdMode/Private/Modes/PlacementModeSubsystem.h
@@ -7,10 +7,7 @@
 
 #include "PlacementModeSubsystem.generated.h"
 
-struct FPaletteItem;
-struct FTypedElementHandle;
 class UAssetPlacementSettings;
-class UPlacementPaletteAsset;
 
 UCLASS()
 class UPlacementModeSubsystem : public UEditorSubsystem
@@ -27,25 +24,6 @@ public:
 	const UAssetPlacementSettings* GetModeSettingsObject() const;
 	UAssetPlacementSettings* GetMutableModeSettingsObject();
 
-	/**
-	 * Verifies if the given element handle is supported by the current mode settings' palette.
-	 *
-	 * @returns true if the element can be placed by the mode.
-	 */
-	bool DoesActivePaletteSupportElement(const FTypedElementHandle& InElementToCheck) const;
-
-	/**
-	 * Adds the given palette item to the current palette.
-	 * 
-	 * @returns a pointer to the added item, if the add was successful.
-	 */
-	FPaletteItem CreatePaletteItem(const FAssetData& InAssetData);
-
-	/**
-	 * Updates the settings object to use the content browser's active selection as the palette.
-	 */
-	void SetUseContentBrowserAsPalette(bool bInUseContentBrowser);
-
 protected:
 	UPROPERTY()
 	TObjectPtr<UAssetPlacementSettings> ModeSettings;
diff --git a/Engine/Source/Editor/AssetPlacementEdMode/Private/PlacementPaletteItem.h b/Engine/Source/Editor/AssetPlacementEdMode/Private/PlacementPaletteItem.h
--- a/Engine/Source/Editor/AssetPlacementEdMode/Private/PlacementPaletteItem.h
+++ b/Engine/Source/Editor/AssetPlacementEdMode/Private/PlacementPaletteItem.h
@@ -7,18 +7,22 @@
 #include "PlacementPaletteItem.generated.h"
 
 class IAssetFactoryInterface;
+class UEditorFactorySettingsObject;
 
 USTRUCT()
 struct FPaletteItem
 {
 	GENERATED_BODY()
 
-	UPROPERTY(EditAnywhere, Category = Placement)
+	UPROPERTY(EditAnywhere, Category = Settings)
 	FSoftObjectPath AssetPath;
 
-	UPROPERTY(EditAnywhere, Category = Placement, AdvancedDisplay)
+	UPROPERTY(EditAnywhere, Category = Settings, AdvancedDisplay)
 	TScriptInterface<IAssetFactoryInterface> AssetFactoryInterface;
 
 	UPROPERTY()
 	FGuid ItemGuid;
+
+	UPROPERTY(EditAnywhere, Instanced, Category = Settings)
+	TObjectPtr<UEditorFactorySettingsObject> SettingsObject = nullptr;
 };
diff --git a/Engine/Source/Editor/AssetPlacementEdMode/Private/SAssetPlacementPalette.cpp b/Engine/Source/Editor/AssetPlacementEdMode/Private/SAssetPlacementPalette.cpp
--- a/Engine/Source/Editor/AssetPlacementEdMode/Private/SAssetPlacementPalette.cpp
+++ b/Engine/Source/Editor/AssetPlacementEdMode/Private/SAssetPlacementPalette.cpp
@@ -12,7 +12,6 @@
 #include "Widgets/Input/SCheckBox.h"
 #include "EditorStyleSet.h"
 #include "AssetThumbnail.h"
-#include "PropertyEditorModule.h"
 
 #include "IContentBrowserSingleton.h"
 #include "ContentBrowserModule.h"
@@ -39,6 +38,7 @@
 #include "PlacementPaletteAsset.h"
 #include "AssetRegistry/AssetRegistryModule.h"
 #include "PlacementPaletteItem.h"
+#include "Factories/AssetFactoryInterface.h"
 
 #define LOCTEXT_NAMESPACE "AssetPlacementMode"
 
@@ -47,6 +47,7 @@
 ////////////////////////////////////////////////
 void SAssetPlacementPalette::Construct(const FArguments& InArgs)
 {
+	ItemDetailsWidget = InArgs._ItemDetailsView;
 	bItemsNeedRebuild = false;
 	bIsRebuildTimerRegistered = false;
 	bIsRefreshTimerRegistered = false;
@@ -251,16 +252,19 @@ void SAssetPlacementPalette::RefreshActivePaletteViewWidget()
 void SAssetPlacementPalette::AddPlacementType(const FAssetData& AssetData)
 {
 	// Try to add the item to the mode's palette
-	FPaletteItem PlacementInfo = GEditor->GetEditorSubsystem<UPlacementModeSubsystem>()->CreatePaletteItem(AssetData);
-	if (!PlacementInfo.AssetPath.IsValid())
+	if (UPlacementModeSubsystem* PlacementModeSubsystem = GEditor->GetEditorSubsystem<UPlacementModeSubsystem>())
 	{
-		return;
-	}
+		FPaletteItem PlacementInfo = PlacementModeSubsystem->GetMutableModeSettingsObject()->AddItemToActivePalette(AssetData);
+		if (!PlacementInfo.AssetPath.IsValid())
+		{
+			return;
+		}
 
-	// Try to load the asset async so it's ready to place.
-	UAssetManager::GetStreamableManager().RequestAsyncLoad(AssetData.ToSoftObjectPath());
+		// Try to load the asset async so it's ready to place
+		UAssetManager::GetStreamableManager().RequestAsyncLoad(AssetData.ToSoftObjectPath());
 
-	PaletteItems.Add(MakeShared<FAssetPlacementPaletteItemModel>(AssetData, SharedThis(this), ThumbnailPool));
+		PaletteItems.Add(MakeShared<FAssetPlacementPaletteItemModel>(AssetData, PlacementInfo, SharedThis(this), ThumbnailPool));
+	}
 }
 
 void SAssetPlacementPalette::ClearPalette()
@@ -270,6 +274,11 @@ void SAssetPlacementPalette::ClearPalette()
 		ModeSubsystem->GetMutableModeSettingsObject()->ClearActivePaletteItems();
 	}
 
+	if (TSharedPtr<IDetailsView> PinnedItemDetails = ItemDetailsWidget.Pin())
+	{
+		PinnedItemDetails->SetObject(nullptr);
+	}
+
 	PaletteItems.Empty();
 	FilteredItems.Empty();
 }
@@ -298,6 +307,8 @@ TSharedRef<SWidgetSwitcher> SAssetPlacementPalette::CreatePaletteViews()
 		.ListItemsSource(&FilteredItems)
 		.OnGenerateTile(this, &SAssetPlacementPalette::GenerateTile)
 		.OnContextMenuOpening(this, &SAssetPlacementPalette::ConstructPlacementTypeContextMenu)
+		.SelectionMode(ESelectionMode::Multi)
+		.OnSelectionChanged(this, &SAssetPlacementPalette::OnPaletteSelectionChanged)
 		.ItemHeight(this, &SAssetPlacementPalette::GetScaledThumbnailSize)
 		.ItemWidth(this, &SAssetPlacementPalette::GetScaledThumbnailSize)
 		.ItemAlignment(EListItemAlignment::LeftAligned);
@@ -307,6 +318,8 @@ TSharedRef<SWidgetSwitcher> SAssetPlacementPalette::CreatePaletteViews()
 		.TreeItemsSource(&FilteredItems)
 		.OnGenerateRow(this, &SAssetPlacementPalette::TreeViewGenerateRow)
 		.OnGetChildren(this, &SAssetPlacementPalette::TreeViewGetChildren)
+		.SelectionMode(ESelectionMode::Multi)
+		.OnSelectionChanged(this, &SAssetPlacementPalette::OnPaletteSelectionChanged)
 		.OnContextMenuOpening(this, &SAssetPlacementPalette::ConstructPlacementTypeContextMenu)
 		.HeaderRow
 		(
@@ -403,7 +416,7 @@ void SAssetPlacementPalette::OnContentBrowserMirrorButtonClicked(ECheckBoxState
 
 	if (UPlacementModeSubsystem* PlacementModeSubsystem = GEditor->GetEditorSubsystem<UPlacementModeSubsystem>())
 	{
-		PlacementModeSubsystem->SetUseContentBrowserAsPalette(bIsMirroringContentBrowser);
+		PlacementModeSubsystem->GetMutableModeSettingsObject()->bUseContentBrowserSelection = bIsMirroringContentBrowser;
 	}
 }
 
@@ -468,6 +481,10 @@ void SAssetPlacementPalette::SetPaletteItems(TArrayView<const FPaletteItem> InPa
 {
 	PaletteItems.Empty();
 	FilteredItems.Empty();
+	if (TSharedPtr<IDetailsView> PinnedDetailsView = ItemDetailsWidget.Pin())
+	{
+		PinnedDetailsView->SetObject(nullptr);
+	}
 
 	FAssetRegistryModule& AssetRegistry = FModuleManager::Get().LoadModuleChecked<FAssetRegistryModule>(TEXT("AssetRegistry"));
 	TArray<FSoftObjectPath> AssetsToLoad;
@@ -476,7 +493,7 @@ void SAssetPlacementPalette::SetPaletteItems(TArrayView<const FPaletteItem> InPa
 		FAssetData AssetData = AssetRegistry.Get().GetAssetByObjectPath(PaletteItem.AssetPath.GetAssetPathName());
 		if (AssetData.IsValid())
 		{
-			PaletteItems.Add(MakeShared<FAssetPlacementPaletteItemModel>(AssetData, SharedThis(this), ThumbnailPool));
+			PaletteItems.Add(MakeShared<FAssetPlacementPaletteItemModel>(AssetData, PaletteItem, SharedThis(this), ThumbnailPool));
 			AssetsToLoad.Emplace(PaletteItem.AssetPath);
 		}
 	}
@@ -684,6 +701,41 @@ TSharedPtr<SWidget> SAssetPlacementPalette::ConstructPlacementTypeContextMenu()
 	return MenuBuilder.MakeWidget();
 }
 
+void SAssetPlacementPalette::OnPaletteSelectionChanged(FPlacementPaletteItemModelPtr Item, ESelectInfo::Type SelectInfo)
+{
+	if (SelectInfo == ESelectInfo::Direct)
+	{
+		return;
+	}
+
+	// Sync selection to the widgets which are not active
+	switch (ActiveViewMode)
+	{
+		case EViewMode::Thumbnail:
+			TreeViewWidget->SetItemSelection(Item, TileViewWidget->IsItemSelected(Item), ESelectInfo::Direct);
+		break;
+
+		case EViewMode::Tree:
+			TileViewWidget->SetItemSelection(Item, TreeViewWidget->IsItemSelected(Item), ESelectInfo::Direct);
+		break;
+	}
+
+	if (TSharedPtr<IDetailsView> PinnedItemDetails = ItemDetailsWidget.Pin())
+	{
+		TArray<UObject*> DetailsObjects;
+		for (const FPlacementPaletteItemModelPtr& SelectedItem : GetActiveViewWidget()->GetSelectedItems())
+		{
+			if (UEditorFactorySettingsObject* SelectedItemSettingsObject = SelectedItem->GetTypeUIInfo()->SettingsObject.Get())
+			{
+				DetailsObjects.Add(SelectedItemSettingsObject);
+			}
+		}
+
+		constexpr bool bForceRefresh = true;
+		PinnedItemDetails->SetObjects(DetailsObjects, true);
+	}
+}
+
 void SAssetPlacementPalette::OnShowPlacementTypeInCB()
 {
 	TArray<FAssetData> FilteredAssets;
diff --git a/Engine/Source/Editor/AssetPlacementEdMode/Private/SAssetPlacementPalette.h b/Engine/Source/Editor/AssetPlacementEdMode/Private/SAssetPlacementPalette.h
--- a/Engine/Source/Editor/AssetPlacementEdMode/Private/SAssetPlacementPalette.h
+++ b/Engine/Source/Editor/AssetPlacementEdMode/Private/SAssetPlacementPalette.h
@@ -44,6 +44,7 @@ public:
 	};
 
 	SLATE_BEGIN_ARGS(SAssetPlacementPalette) {}
+		SLATE_ARGUMENT(TWeakPtr<IDetailsView>, ItemDetailsView)
 	SLATE_END_ARGS()
 
 	void Construct(const FArguments& InArgs);
@@ -100,7 +101,6 @@ private:
 	void SetupContentBrowserMirroring(bool bInMirrorContentBrowser);
 	void OnSetPaletteAsset(const FAssetData& InAssetData);
 	void SetPaletteItems(TArrayView<const FPaletteItem> InPaletteItems);
-	void SetToLocalUserPalette();
 	FString GetPalettePath() const;
 
 	/** Sets the view mode of the palette */
@@ -130,6 +130,9 @@ private:	// CONTEXT MENU
 	/** @return the SWidget containing the context menu */
 	TSharedPtr<SWidget> ConstructPlacementTypeContextMenu();
 
+	/** Called when the selected set of items changes */
+	void OnPaletteSelectionChanged(FPlacementPaletteItemModelPtr Item, ESelectInfo::Type SelectInfo);
+
 	/** Handler for 'Show in CB' command  */
 	void OnShowPlacementTypeInCB();
 
@@ -193,7 +196,7 @@ private:
 	TSharedPtr<SPlacementTypeTreeView> TreeViewWidget;
 
 	/** Placement mesh details widget  */
-	TSharedPtr<IDetailsView> DetailsWidget;
+	TWeakPtr<IDetailsView> ItemDetailsWidget;
 
 	/** Placement items search box widget */
 	TSharedPtr<SSearchBox> SearchBoxPtr;
diff --git a/Engine/Source/Editor/AssetPlacementEdMode/Private/Tools/PlacementBrushToolBase.cpp b/Engine/Source/Editor/AssetPlacementEdMode/Private/Tools/PlacementBrushToolBase.cpp
--- a/Engine/Source/Editor/AssetPlacementEdMode/Private/Tools/PlacementBrushToolBase.cpp
+++ b/Engine/Source/Editor/AssetPlacementEdMode/Private/Tools/PlacementBrushToolBase.cpp
@@ -278,6 +278,12 @@ FTypedElementListRef UPlacementBrushToolBase::GetElementsInBrushRadius(const FIn
 		return ElementHandles;
 	}
 
+	UPlacementModeSubsystem* PlacementModeSubsystem = GEditor->GetEditorSubsystem<UPlacementModeSubsystem>();
+	if (!PlacementModeSubsystem || !PlacementModeSubsystem->GetModeSettingsObject())
+	{
+		return ElementHandles;
+	}
+
 	FCollisionQueryParams QueryParams(TEXT("PlacementBrushTool"), SCENE_QUERY_STAT_ONLY(EdMode_PlacementTrace), true);
 	QueryParams.bReturnFaceIndex = false;
 	TArray<FHitResult> Hits;
@@ -295,7 +301,7 @@ FTypedElementListRef UPlacementBrushToolBase::GetElementsInBrushRadius(const FIn
 		if (UInstancedStaticMeshComponent* ISMComponent = Cast<UInstancedStaticMeshComponent>(Hit.GetComponent()))
 		{
 			FTypedElementHandle StaticMeshObjectHandle = UEngineElementsLibrary::AcquireEditorObjectElementHandle(ISMComponent->GetStaticMesh());
-			if (!GEditor->GetEditorSubsystem<UPlacementModeSubsystem>()->DoesActivePaletteSupportElement(StaticMeshObjectHandle))
+			if (!PlacementModeSubsystem->GetModeSettingsObject()->DoesActivePaletteSupportElement(StaticMeshObjectHandle))
 			{
 				continue;
 			}
@@ -313,7 +319,7 @@ FTypedElementListRef UPlacementBrushToolBase::GetElementsInBrushRadius(const FIn
 		else
 		{
 			FTypedElementHandle PrimarySelectionHandle = SelectionSet->GetSelectionElement(UEngineElementsLibrary::AcquireEditorComponentElementHandle(Hit.GetComponent()), ETypedElementSelectionMethod::Primary);
-			if (!GEditor->GetEditorSubsystem<UPlacementModeSubsystem>()->DoesActivePaletteSupportElement(PrimarySelectionHandle))
+			if (!PlacementModeSubsystem->GetModeSettingsObject()->DoesActivePaletteSupportElement(PrimarySelectionHandle))
 			{
 				continue;
 			}
diff --git a/Engine/Source/Editor/AssetPlacementEdMode/Private/Tools/PlacementPlaceSingleTool.cpp b/Engine/Source/Editor/AssetPlacementEdMode/Private/Tools/PlacementPlaceSingleTool.cpp
--- a/Engine/Source/Editor/AssetPlacementEdMode/Private/Tools/PlacementPlaceSingleTool.cpp
+++ b/Engine/Source/Editor/AssetPlacementEdMode/Private/Tools/PlacementPlaceSingleTool.cpp
@@ -251,6 +251,7 @@ void UPlacementModePlaceSingleTool::GeneratePlacementData(const FInputDeviceRay&
 			PlacementInfo->ItemGuid = ItemToPlace.ItemGuid;
 			PlacementInfo->FinalizedTransform = FinalizeTransform(TransformToUpdate, LastBrushStamp.WorldNormal, PlacementSettings);
 			PlacementInfo->PreferredLevel = GEditor->GetEditorWorldContext().World()->GetCurrentLevel();
+			PlacementInfo->SettingsObject = ItemToPlace.SettingsObject;
 		}
 	}
 }
diff --git a/Engine/Source/Editor/AssetPlacementEdMode/Private/Tools/PlacementPlaceTool.cpp b/Engine/Source/Editor/AssetPlacementEdMode/Private/Tools/PlacementPlaceTool.cpp
--- a/Engine/Source/Editor/AssetPlacementEdMode/Private/Tools/PlacementPlaceTool.cpp
+++ b/Engine/Source/Editor/AssetPlacementEdMode/Private/Tools/PlacementPlaceTool.cpp
@@ -67,40 +67,43 @@ void UPlacementModePlacementTool::OnTick(float DeltaTime)
 				FVector Start, End;
 				GetRandomVectorInBrush(Start, End);
 				FHitResult AdjustedHitResult;
-				FindHitResultWithStartAndEndTraceVectors(AdjustedHitResult, Start, End);
-				FVector SpawnLocation = AdjustedHitResult.ImpactPoint;
-				FVector SpawnNormal = AdjustedHitResult.ImpactNormal;
-
-				bool bValidInstance = true;
-				auto TempInstances = PotentialInstanceHash.GetInstancesOverlappingBox(FBox::BuildAABB(BrushLocation, FVector(BrushProperties->BrushRadius)));
-				for (int32 InstanceIndex : TempInstances)
+				if (FindHitResultWithStartAndEndTraceVectors(AdjustedHitResult, Start, End))
 				{
-					if ((PotentialInstanceLocations[InstanceIndex] - BrushLocation).SizeSquared() < FMath::Square(BrushProperties->BrushRadius))
+					FVector SpawnLocation = AdjustedHitResult.ImpactPoint;
+					FVector SpawnNormal = AdjustedHitResult.ImpactNormal;
+
+					bool bValidInstance = true;
+					auto TempInstances = PotentialInstanceHash.GetInstancesOverlappingBox(FBox::BuildAABB(BrushLocation, FVector(BrushProperties->BrushRadius)));
+					for (int32 InstanceIndex : TempInstances)
 					{
-						bValidInstance = false;
-						break;
+						if ((PotentialInstanceLocations[InstanceIndex] - BrushLocation).SizeSquared() < FMath::Square(BrushProperties->BrushRadius))
+						{
+							bValidInstance = false;
+							break;
+						}
 					}
-				}
 
-				if (!bValidInstance)
-				{
-					continue;
-				}
+					if (!bValidInstance)
+					{
+						continue;
+					}
 
-				int32 PotentialIdx = PotentialInstanceLocations.Add(SpawnLocation);
-				PotentialInstanceHash.InsertInstance(SpawnLocation, PotentialIdx);
+					int32 PotentialIdx = PotentialInstanceLocations.Add(SpawnLocation);
+					PotentialInstanceHash.InsertInstance(SpawnLocation, PotentialIdx);
 
-				int32 ItemIndex = FMath::RandHelper(PaletteItems.Num());
-				const FPaletteItem& ItemToPlace = PaletteItems[ItemIndex];
+					int32 ItemIndex = FMath::RandHelper(PaletteItems.Num());
+					const FPaletteItem& ItemToPlace = PaletteItems[ItemIndex];
 
-				FAssetPlacementInfo NewInfo;
-				NewInfo.AssetToPlace = ItemToPlace.AssetPath.TryLoad();
-				NewInfo.FactoryOverride = ItemToPlace.AssetFactoryInterface;
-				NewInfo.ItemGuid = ItemToPlace.ItemGuid;
-				NewInfo.PreferredLevel = GEditor->GetEditorWorldContext().World()->GetCurrentLevel();
-				NewInfo.FinalizedTransform = GenerateTransformFromHitLocationAndNormal(SpawnLocation, SpawnNormal);
+					FAssetPlacementInfo NewInfo;
+					NewInfo.AssetToPlace = ItemToPlace.AssetPath.TryLoad();
+					NewInfo.FactoryOverride = ItemToPlace.AssetFactoryInterface;
+					NewInfo.ItemGuid = ItemToPlace.ItemGuid;
+					NewInfo.PreferredLevel = GEditor->GetEditorWorldContext().World()->GetCurrentLevel();
+					NewInfo.FinalizedTransform = GenerateTransformFromHitLocationAndNormal(SpawnLocation, SpawnNormal);
+					NewInfo.SettingsObject = ItemToPlace.SettingsObject;
 
-				DesiredInfoSpawnLocations.Emplace(NewInfo);
+					DesiredInfoSpawnLocations.Emplace(NewInfo);
+				}
 			}
 
 			UPlacementSubsystem* PlacementSubsystem = GEditor->GetEditorSubsystem<UPlacementSubsystem>();
diff --git a/Engine/Source/Editor/AssetPlacementEdMode/Public/AssetPlacementEdModeToolkit.h b/Engine/Source/Editor/AssetPlacementEdMode/Public/AssetPlacementEdModeToolkit.h
--- a/Engine/Source/Editor/AssetPlacementEdMode/Public/AssetPlacementEdModeToolkit.h
+++ b/Engine/Source/Editor/AssetPlacementEdMode/Public/AssetPlacementEdModeToolkit.h
@@ -22,11 +22,15 @@ public:
 	virtual FText GetBaseToolkitName() const override;
 	virtual void InvokeUI() override;
 protected:
-	virtual TSharedPtr<SWidget> GetInlineContent() const override;
 	virtual void RequestModeUITabs() override;
 	TSharedRef<SDockTab> CreateAssetPalette(const FSpawnTabArgs& Args);
+	TSharedRef<SDockTab> CreatePaletteItemDetailsView(const FSpawnTabArgs& Args);
+	EVisibility GetDetailsViewHintTextVisibility() const;
 
 private:
 	TWeakPtr<SDockTab> AssetPaletteTab;
+	TWeakPtr<SDockTab> PaletteItemDetailsViewTab;
 	FMinorTabConfig AssetPaletteInfo;
+	FMinorTabConfig PaletteItemDetailsViewInfo;
+	TSharedPtr<IDetailsView> PaletteItemDetailsView;
 };
diff --git a/Engine/Source/Editor/EditorFramework/Private/Subsystems/PlacementSubsystem.cpp b/Engine/Source/Editor/EditorFramework/Private/Subsystems/PlacementSubsystem.cpp
--- a/Engine/Source/Editor/EditorFramework/Private/Subsystems/PlacementSubsystem.cpp
+++ b/Engine/Source/Editor/EditorFramework/Private/Subsystems/PlacementSubsystem.cpp
@@ -96,13 +96,6 @@ bool UPlacementSubsystem::IsCreatingPreviewElements() const
 	return bIsCreatingPreviewElements;
 }
 
-FGuid UPlacementSubsystem::GetUserGridGuid()
-{
-	FGuid UserGridGuid;
-	FGuid::Parse(FPlatformMisc::GetLoginId(), UserGridGuid);
-	return UserGridGuid;
-}
-
 void UPlacementSubsystem::RegisterPlacementFactories()
 {
 	for (TObjectIterator<UClass> ObjectIt; ObjectIt; ++ObjectIt)
diff --git a/Engine/Source/Editor/EditorFramework/Public/Factories/AssetFactoryInterface.h b/Engine/Source/Editor/EditorFramework/Public/Factories/AssetFactoryInterface.h
--- a/Engine/Source/Editor/EditorFramework/Public/Factories/AssetFactoryInterface.h
+++ b/Engine/Source/Editor/EditorFramework/Public/Factories/AssetFactoryInterface.h
@@ -12,6 +12,12 @@ class IAssetFactoryInterface;
 struct FPlacementOptions;
 struct FAssetPlacementInfo;
 
+UCLASS(Abstract, hideCategories=Object, editinlinenew, collapsecategories, MinimalAPI)
+class UEditorFactorySettingsObject : public UObject
+{
+	GENERATED_BODY()
+};
+
 UINTERFACE(MinimalAPI)
 class UAssetFactoryInterface : public UInterface
 {
@@ -74,4 +80,9 @@ public:
 	 * For example, for instanced static mesh placement, we'd rebuild the parent component's tree here.
 	 */
 	 virtual void EndPlacement(TArrayView<const FTypedElementHandle> InPlacedElements, const FPlacementOptions& InPlacementOptions) = 0;
+
+	 /**
+	  * Returns the settings object which this factory will use to customize placement settings, based on the given placement information.
+	  */
+	 virtual UEditorFactorySettingsObject* FactorySettingsObjectForPlacement(const FAssetData& InAssetData, const FPlacementOptions& InPlacementOptions) = 0;
 };
diff --git a/Engine/Source/Editor/EditorFramework/Public/Subsystems/PlacementSubsystem.h b/Engine/Source/Editor/EditorFramework/Public/Subsystems/PlacementSubsystem.h
--- a/Engine/Source/Editor/EditorFramework/Public/Subsystems/PlacementSubsystem.h
+++ b/Engine/Source/Editor/EditorFramework/Public/Subsystems/PlacementSubsystem.h
@@ -10,6 +10,7 @@
 #include "PlacementSubsystem.generated.h"
 
 class IAssetFactoryInterface;
+class UEditorFactorySettingsObject;
 
 USTRUCT()
 struct EDITORFRAMEWORK_API FAssetPlacementInfo
@@ -43,6 +44,9 @@ struct EDITORFRAMEWORK_API FAssetPlacementInfo
 	 */
 	UPROPERTY()
 	FGuid ItemGuid;
+
+	UPROPERTY()
+	TObjectPtr<UEditorFactorySettingsObject> SettingsObject = nullptr;
 };
 
 USTRUCT()
@@ -101,11 +105,6 @@ public:
 	 */
 	bool IsCreatingPreviewElements() const;
 
-	/**
-	 * Returns the guid which the placement system should use to tie instanced placement to the current user for partitions.
-	 */
-	static FGuid GetUserGridGuid();
-
 private:
 	void RegisterPlacementFactories();
 	void UnregisterPlacementFactories();
diff --git a/Engine/Source/Editor/UnrealEd/Classes/ActorFactories/ActorFactory.h b/Engine/Source/Editor/UnrealEd/Classes/ActorFactories/ActorFactory.h
--- a/Engine/Source/Editor/UnrealEd/Classes/ActorFactories/ActorFactory.h
+++ b/Engine/Source/Editor/UnrealEd/Classes/ActorFactories/ActorFactory.h
@@ -20,6 +20,24 @@ struct FAssetData;
 class UBlueprint;
 class ULevel;
 
+UENUM()
+enum class EActorPlacementType
+{
+	FullActor,	// Places a full actor, one per instance
+	//InstancedStaticMesh,	// Places the actor as a set of instanced static meshes. All other components will be ignored.
+	//LightweightInstance,	// Places the actor as a lightweight instanced static mesh, which can be hydrated at runtime.
+};
+
+UCLASS()
+class UActorFactoryPlacementSettings : public UEditorFactorySettingsObject
+{
+	GENERATED_BODY()
+
+public:
+	UPROPERTY(VisibleAnywhere, Category="Factory Settings")
+	EActorPlacementType ActorPlacementType = EActorPlacementType::FullActor;
+};
+
 UCLASS(collapsecategories, hidecategories=Object, editinlinenew, config=Editor, abstract, transient)
 class UNREALED_API UActorFactory : public UObject, public IAssetFactoryInterface
 {
@@ -90,6 +108,7 @@ class UNREALED_API UActorFactory : public UObject, public IAssetFactoryInterface
 	virtual FAssetData GetAssetDataFromElementHandle(const FTypedElementHandle& InHandle) override;
 	virtual void BeginPlacement(const FPlacementOptions& InPlacementOptions) override;
 	virtual void EndPlacement(TArrayView<const FTypedElementHandle> InPlacedElements, const FPlacementOptions& InPlacementOptions) override;
+	virtual UEditorFactorySettingsObject* FactorySettingsObjectForPlacement(const FAssetData& InAssetData, const FPlacementOptions& InPlacementOptions) override;
 	// End IAssetFactoryInterface Interface
 
 protected:
diff --git a/Engine/Source/Editor/UnrealEd/Private/Factories/ActorFactory.cpp b/Engine/Source/Editor/UnrealEd/Private/Factories/ActorFactory.cpp
--- a/Engine/Source/Editor/UnrealEd/Private/Factories/ActorFactory.cpp
+++ b/Engine/Source/Editor/UnrealEd/Private/Factories/ActorFactory.cpp
@@ -382,6 +382,12 @@ void UActorFactory::EndPlacement(TArrayView<const FTypedElementHandle> InPlacedE
 {
 }
 
+UEditorFactorySettingsObject* UActorFactory::FactorySettingsObjectForPlacement(const FAssetData& InAssetData, const FPlacementOptions& InPlacementOptions)
+{
+	UActorFactoryPlacementSettings* FactorySettingsObject = NewObject<UActorFactoryPlacementSettings>(this);
+	return FactorySettingsObject;
+}
+
 AActor* UActorFactory::CreateActor( UObject* Asset, ULevel* InLevel, const FTransform& SpawnTransform, EObjectFlags InObjectFlags, const FName Name)
 {
 	FActorSpawnParameters SpawnParams;
diff --git a/Engine/Source/Editor/UnrealEd/Private/Factories/EditorStaticMeshFactory.cpp b/Engine/Source/Editor/UnrealEd/Private/Factories/EditorStaticMeshFactory.cpp
--- a/Engine/Source/Editor/UnrealEd/Private/Factories/EditorStaticMeshFactory.cpp
+++ b/Engine/Source/Editor/UnrealEd/Private/Factories/EditorStaticMeshFactory.cpp
@@ -17,6 +17,13 @@
 
 #include "Subsystems/PlacementSubsystem.h"
 
+void UEditorStaticMeshFactoryPlacementSettings::PostEditChangeProperty(struct FPropertyChangedEvent& PropertyChangedEvent)
+{
+	Super::PostEditChangeProperty(PropertyChangedEvent);
+
+	StaticMeshComponentDescriptor.ComputeHash();
+}
+
 TArray<FTypedElementHandle> UEditorStaticMeshFactory::PlaceAsset(const FAssetPlacementInfo& InPlacementInfo, const FPlacementOptions& InPlacementOptions)
 {
 	// If we're disallowing instanced placement, or creating preview elements, don't use the ISM placement.
@@ -28,30 +35,18 @@ TArray<FTypedElementHandle> UEditorStaticMeshFactory::PlaceAsset(const FAssetPla
 	TArray<FTypedElementHandle> PlacedInstanceHandles;
 	if (InPlacementInfo.PreferredLevel.IsValid())
 	{
-		UObject* AssetToPlaceAsObject = InPlacementInfo.AssetToPlace.GetAsset();
-		FISMComponentDescriptor ComponentDescriptor;
-		if (UStaticMesh* StaticMeshObject = Cast<UStaticMesh>(AssetToPlaceAsObject))
-		{
-			// If this is a Nanite mesh, prefer to use ISM over HISM, as HISM duplicates many features/bookkeeping that Nanite already handles for us.
-			if (StaticMeshObject->HasValidNaniteData())
-			{
-				ComponentDescriptor.InitFrom(UInstancedStaticMeshComponent::StaticClass()->GetDefaultObject<UInstancedStaticMeshComponent>());
-			}
-			ComponentDescriptor.StaticMesh = StaticMeshObject;
-		}
-		else if (AStaticMeshActor* StaticMeshActor = Cast<AStaticMeshActor>(AssetToPlaceAsObject))
+		if (!InPlacementInfo.SettingsObject)
 		{
-			if (UStaticMeshComponent* StaticMeshComponent = StaticMeshActor->GetStaticMeshComponent())
-			{
-				ComponentDescriptor.StaticMesh = StaticMeshComponent->GetStaticMesh();
-			}
+			return TArray<FTypedElementHandle>();
 		}
 
+		FISMComponentDescriptor ComponentDescriptor = CastChecked<UEditorStaticMeshFactoryPlacementSettings>(InPlacementInfo.SettingsObject)->StaticMeshComponentDescriptor;
 		if (!ComponentDescriptor.StaticMesh)
 		{
 			return TArray<FTypedElementHandle>();
 		}
 
+		// Make sure the component descriptor's hash matches its current settings before we place.
 		ComponentDescriptor.ComputeHash();
 
 		if (UActorPartitionSubsystem* PartitionSubsystem = UWorld::GetSubsystem<UActorPartitionSubsystem>(InPlacementInfo.PreferredLevel->GetWorld()))
@@ -133,6 +128,42 @@ FAssetData UEditorStaticMeshFactory::GetAssetDataFromElementHandle(const FTypedE
 	return Super::GetAssetDataFromElementHandle(InHandle);
 }
 
+UEditorFactorySettingsObject* UEditorStaticMeshFactory::FactorySettingsObjectForPlacement(const FAssetData& InAssetData, const FPlacementOptions& InPlacementOptions)
+{
+	if (!ShouldPlaceInstancedStaticMeshes(InPlacementOptions))
+	{
+		return Super::FactorySettingsObjectForPlacement(InAssetData, InPlacementOptions);
+	}
+
+	UEditorStaticMeshFactoryPlacementSettings* PlacementSettingsObject = NewObject<UEditorStaticMeshFactoryPlacementSettings>(this);
+	if (PlacementSettingsObject)
+	{
+		UObject* AssetToPlaceAsObject = InAssetData.GetAsset();
+		FISMComponentDescriptor& ComponentDescriptor = PlacementSettingsObject->StaticMeshComponentDescriptor;
+		if (UStaticMesh* StaticMeshObject = Cast<UStaticMesh>(AssetToPlaceAsObject))
+		{
+			// If this is a Nanite mesh, prefer to use ISM over HISM, as HISM duplicates many features/bookkeeping that Nanite already handles for us.
+			if (StaticMeshObject->HasValidNaniteData())
+			{
+				ComponentDescriptor.InitFrom(UInstancedStaticMeshComponent::StaticClass()->GetDefaultObject<UInstancedStaticMeshComponent>());
+			}
+			ComponentDescriptor.StaticMesh = StaticMeshObject;
+		}
+		else if (AStaticMeshActor* StaticMeshActor = Cast<AStaticMeshActor>(AssetToPlaceAsObject))
+		{
+			if (UStaticMeshComponent* StaticMeshComponent = StaticMeshActor->GetStaticMeshComponent())
+			{
+				ComponentDescriptor.StaticMesh = StaticMeshComponent->GetStaticMesh();
+			}
+		}
+
+		// Go ahead and compute the descriptor now, in case we do not go through a place cycle or edit any properties
+		ComponentDescriptor.ComputeHash();
+	}
+
+	return PlacementSettingsObject;
+}
+
 bool UEditorStaticMeshFactory::ShouldPlaceInstancedStaticMeshes(const FPlacementOptions& InPlacementOptions) const
 {
 	return !InPlacementOptions.bIsCreatingPreviewElements && InPlacementOptions.InstancedPlacementGridGuid.IsValid();
diff --git a/Engine/Source/Editor/UnrealEd/Private/Factories/EditorStaticMeshFactory.h b/Engine/Source/Editor/UnrealEd/Private/Factories/EditorStaticMeshFactory.h
--- a/Engine/Source/Editor/UnrealEd/Private/Factories/EditorStaticMeshFactory.h
+++ b/Engine/Source/Editor/UnrealEd/Private/Factories/EditorStaticMeshFactory.h
@@ -4,11 +4,24 @@
 
 #include "Factories/AssetFactoryInterface.h"
 #include "ActorFactories/ActorFactoryStaticMesh.h"
+#include "ISMPartition/ISMComponentDescriptor.h"
 
 #include "EditorStaticMeshFactory.generated.h"
 
 class AISMPartitionActor;
 
+UCLASS()
+class UEditorStaticMeshFactoryPlacementSettings : public UEditorFactorySettingsObject
+{
+	GENERATED_BODY()
+
+public:
+	UPROPERTY(EditAnywhere, Category = "ISM Component Settings", meta=(ShowOnlyInnerProperties))
+	FISMComponentDescriptor StaticMeshComponentDescriptor;
+
+	virtual void PostEditChangeProperty(struct FPropertyChangedEvent& PropertyChangedEvent) override;
+};
+
 UCLASS(Transient)
 class UEditorStaticMeshFactory : public UActorFactoryStaticMesh
 {
@@ -19,6 +32,7 @@ public:
 	virtual TArray<FTypedElementHandle> PlaceAsset(const FAssetPlacementInfo& InPlacementInfo, const FPlacementOptions& InPlacementOptions) override;
 	virtual FAssetData GetAssetDataFromElementHandle(const FTypedElementHandle& InHandle) override;
 	virtual void EndPlacement(TArrayView<const FTypedElementHandle> InPlacedElements, const FPlacementOptions& InPlacementOptions) override;
+	virtual UEditorFactorySettingsObject* FactorySettingsObjectForPlacement(const FAssetData& InAssetData, const FPlacementOptions& InPlacementOptions) override;
 	// End IAssetFactoryInterface
 
 protected:
diff --git a/Engine/Source/Runtime/Engine/Public/ISMPartition/ISMComponentDescriptor.h b/Engine/Source/Runtime/Engine/Public/ISMPartition/ISMComponentDescriptor.h
--- a/Engine/Source/Runtime/Engine/Public/ISMPartition/ISMComponentDescriptor.h
+++ b/Engine/Source/Runtime/Engine/Public/ISMPartition/ISMComponentDescriptor.h
@@ -52,121 +52,121 @@ public:
 	UPROPERTY()
 	mutable uint32 Hash = 0;
 
-	UPROPERTY()
+	UPROPERTY(EditAnywhere, Category = "Component Settings")
 	TSubclassOf<UInstancedStaticMeshComponent> ComponentClass;
 	
-	UPROPERTY()
+	UPROPERTY(VisibleAnywhere, Category = "Component Settings")
 	TObjectPtr<UStaticMesh> StaticMesh = nullptr;
 	
-	UPROPERTY()
+	UPROPERTY(EditAnywhere, Category = "Component Settings")
 	TArray<TObjectPtr<UMaterialInterface>> OverrideMaterials;
 	
-	UPROPERTY()
+	UPROPERTY(EditAnywhere, Category = "Component Settings")
 	TArray<TObjectPtr<URuntimeVirtualTexture>> RuntimeVirtualTextures;
 
-	UPROPERTY()
+	UPROPERTY(EditAnywhere, Category = "Component Settings")
 	TEnumAsByte<EComponentMobility::Type> Mobility;
 		
-	UPROPERTY()
+	UPROPERTY(EditAnywhere, Category = "Component Settings")
 	ERuntimeVirtualTextureMainPassType VirtualTextureRenderPassType;
 	
-	UPROPERTY()
+	UPROPERTY(EditAnywhere, Category = "Component Settings")
 	ELightmapType LightmapType;
 
-	UPROPERTY()
+	UPROPERTY(EditAnywhere, Category = "Component Settings")
 	FLightingChannels LightingChannels;
 
-	UPROPERTY()
+	UPROPERTY(EditAnywhere, Category = "Component Settings")
 	int32 RayTracingGroupId;
 
-	UPROPERTY()
+	UPROPERTY(EditAnywhere, Category = "Component Settings")
 	ERayTracingGroupCullingPriority RayTracingGroupCullingPriority;
 
-	UPROPERTY()
+	UPROPERTY(EditAnywhere, Category = "Component Settings")
 	TEnumAsByte<EHasCustomNavigableGeometry::Type> bHasCustomNavigableGeometry;
 
-	UPROPERTY()
+	UPROPERTY(EditAnywhere, Category = "Component Settings")
 	ERendererStencilMask CustomDepthStencilWriteMask;
 
-	UPROPERTY()
+	UPROPERTY(EditAnywhere, Category = "Component Settings")
 	FBodyInstance BodyInstance;
 		
-	UPROPERTY()
+	UPROPERTY(EditAnywhere, Category = "Component Settings")
 	int32 InstanceStartCullDistance;
 
-	UPROPERTY()
+	UPROPERTY(EditAnywhere, Category = "Component Settings")
 	int32 InstanceEndCullDistance;
 
-	UPROPERTY()
+	UPROPERTY(EditAnywhere, Category = "Component Settings")
 	int32 VirtualTextureCullMips;
 
-	UPROPERTY()
+	UPROPERTY(EditAnywhere, Category = "Component Settings")
 	int32 TranslucencySortPriority;
 
-	UPROPERTY()
+	UPROPERTY(EditAnywhere, Category = "Component Settings")
 	int32 OverriddenLightMapRes;
 
-	UPROPERTY()
+	UPROPERTY(EditAnywhere, Category = "Component Settings")
 	int32 CustomDepthStencilValue;
 
-	UPROPERTY()
+	UPROPERTY(EditAnywhere, Category = "Component Settings")
 	EHLODBatchingPolicy HLODBatchingPolicy;
 	
-	UPROPERTY()
+	UPROPERTY(EditAnywhere, Category = "Component Settings")
 	uint8 bCastShadow : 1;
 
-	UPROPERTY()
+	UPROPERTY(EditAnywhere, Category = "Component Settings")
 	uint8 bCastDynamicShadow : 1;
 
-	UPROPERTY()
+	UPROPERTY(EditAnywhere, Category = "Component Settings")
 	uint8 bCastStaticShadow : 1;
 
-	UPROPERTY()
+	UPROPERTY(EditAnywhere, Category = "Component Settings")
 	uint8 bCastContactShadow : 1;
 
-	UPROPERTY()
+	UPROPERTY(EditAnywhere, Category = "Component Settings")
 	uint8 bCastShadowAsTwoSided : 1;
 
-	UPROPERTY()
+	UPROPERTY(EditAnywhere, Category = "Component Settings")
 	uint8 bAffectDynamicIndirectLighting : 1;
 
-	UPROPERTY()
+	UPROPERTY(EditAnywhere, Category = "Component Settings")
 	uint8 bAffectDistanceFieldLighting : 1;
 
-	UPROPERTY()
+	UPROPERTY(EditAnywhere, Category = "Component Settings")
 	uint8 bReceivesDecals : 1;
 
-	UPROPERTY()
+	UPROPERTY(EditAnywhere, Category = "Component Settings")
 	uint8 bOverrideLightMapRes : 1;
 
-	UPROPERTY()
+	UPROPERTY(EditAnywhere, Category = "Component Settings")
 	uint8 bUseAsOccluder : 1;
 
-	UPROPERTY()
+	UPROPERTY(EditAnywhere, Category = "Component Settings")
 	uint8 bEnableDensityScaling : 1;
 
-	UPROPERTY()
+	UPROPERTY(EditAnywhere, Category = "Component Settings")
 	uint8 bEnableDiscardOnLoad : 1;
 
-	UPROPERTY()
+	UPROPERTY(EditAnywhere, Category = "Component Settings")
 	uint8 bRenderCustomDepth : 1;
 
-	UPROPERTY()
+	UPROPERTY(EditAnywhere, Category = "Component Settings")
 	uint8 bIncludeInHLOD : 1;
 
-	UPROPERTY()
+	UPROPERTY(EditAnywhere, Category = "Component Settings")
 	uint8 bVisibleInRayTracing : 1;
 
-	UPROPERTY()
+	UPROPERTY(EditAnywhere, Category = "Component Settings")
 	uint8 bHiddenInGame : 1;
 
-	UPROPERTY()
+	UPROPERTY(EditAnywhere, Category = "Component Settings")
 	uint8 bIsEditorOnly : 1;
 
-	UPROPERTY()
+	UPROPERTY(EditAnywhere, Category = "Component Settings")
 	uint8 bVisible : 1;
 
-	UPROPERTY()
+	UPROPERTY(EditAnywhere, Category = "Component Settings")
 	uint8 bConsiderForActorPlacementWhenHidden : 1;
 #endif
 };
