Fixed nopch compile error

diff --git a/Engine/Source/Editor/AssetPlacementEdMode/Private/PlacementPaletteItem.h b/Engine/Source/Editor/AssetPlacementEdMode/Private/PlacementPaletteItem.h
--- a/Engine/Source/Editor/AssetPlacementEdMode/Private/PlacementPaletteItem.h
+++ b/Engine/Source/Editor/AssetPlacementEdMode/Private/PlacementPaletteItem.h
@@ -4,6 +4,7 @@
 
 #include "CoreMinimal.h"
 #include "ISMPartition/ISMPartitionInstanceManager.h"
+#include "Factories/AssetFactoryInterface.h"
 
 #include "PlacementPaletteItem.generated.h"
 
